#UE5 Fix case where a component request may not properly add a component on an actor if the actor has initialized, but has not yet begun play (which can happen if you activate a plugin in a beginplay callstack. Add call to RemoveModularComponentReciever in ModularCharacter's Endplay to match the add done in PreInitializeComponents.

#tests PIE in sample project

diff --git a/Engine/Plugins/Experimental/ModularGameplay/Source/ModularGameplay/Private/Components/GameFrameworkComponentManager.cpp b/Engine/Plugins/Experimental/ModularGameplay/Source/ModularGameplay/Private/Components/GameFrameworkComponentManager.cpp
--- a/Engine/Plugins/Experimental/ModularGameplay/Source/ModularGameplay/Private/Components/GameFrameworkComponentManager.cpp
+++ b/Engine/Plugins/Experimental/ModularGameplay/Source/ModularGameplay/Private/Components/GameFrameworkComponentManager.cpp
@@ -198,7 +198,7 @@ TSharedPtr<FComponentRequestHandle> UGameFrameworkComponentManager::AddComponent
 				{
 					for (TActorIterator<AActor> ActorIt(LocalWorld, ReceiverClassPtr); ActorIt; ++ActorIt)
 					{
-						if (ActorIt->HasActorBegunPlay())
+						if (ActorIt->IsActorInitialized())
 						{
 #if WITH_EDITOR
 							ensureMsgf(AllReceivers.Contains(*ActorIt), TEXT("You may not add a component request for an actor class that does not call AddReceiver/RemoveReceiver in code! Class:%s"), *GetPathNameSafe(ReceiverClassPtr));
