Removing some of the cruft left over from the faux game feature system. Plus, adding a new more generic system for extending actor instances.

diff --git a/Engine/Plugins/Experimental/ModularGameplay/Source/ModularGameplay/Public/Components/GameFrameworkComponentManager.h b/Engine/Plugins/Experimental/ModularGameplay/Source/ModularGameplay/Public/Components/GameFrameworkComponentManager.h
--- a/Engine/Plugins/Experimental/ModularGameplay/Source/ModularGameplay/Public/Components/GameFrameworkComponentManager.h
+++ b/Engine/Plugins/Experimental/ModularGameplay/Source/ModularGameplay/Public/Components/GameFrameworkComponentManager.h
@@ -67,10 +67,10 @@ public:
 	MODULARGAMEPLAY_API void RemoveReceiver(AActor* Receiver);
 
 	/** Adds an actor as a receiver for components (automatically finding the manager for the actor's  game instance). If it passes the actorclass filter on requests it will get the components. */
-	static void AddGameFrameworkComponentReceiver(AActor* Receiver, bool bAddOnlyInGameWorlds = true);
+	static MODULARGAMEPLAY_API void AddGameFrameworkComponentReceiver(AActor* Receiver, bool bAddOnlyInGameWorlds = true);
 
 	/** Removes an actor as a receiver for components (automatically finding the manager for the actor's game instance). */
-	static void RemoveGameFrameworkComponentReceiver(AActor* Receiver);
+	static MODULARGAMEPLAY_API void RemoveGameFrameworkComponentReceiver(AActor* Receiver);
 
 	/** Adds a request to instantiate components on actors of the given classes. Returns a handle that will keep the request "alive" until it is destructed, at which point the request is removed. */
 	MODULARGAMEPLAY_API TSharedPtr<FComponentRequestHandle> AddComponentRequest(const TSoftClassPtr<AActor>& ReceiverClass, TSubclassOf<UActorComponent> ComponentClass);
