Change GameStateComponent to depend on GameStateBase, I may move the match started virtual up to AGameStateBase in a different checkin
Made one unchecked accessor safe, and added comment explaining safety for the rest as fixing them required adding a lot of null checks

diff --git a/Engine/Plugins/Experimental/ModularGameplay/Source/ModularGameplay/Public/Components/ControllerComponent.h b/Engine/Plugins/Experimental/ModularGameplay/Source/ModularGameplay/Public/Components/ControllerComponent.h
--- a/Engine/Plugins/Experimental/ModularGameplay/Source/ModularGameplay/Public/Components/ControllerComponent.h
+++ b/Engine/Plugins/Experimental/ModularGameplay/Source/ModularGameplay/Public/Components/ControllerComponent.h
@@ -19,6 +19,7 @@ public:
 	
 	UControllerComponent(const FObjectInitializer& ObjectInitializer);
 
+	/** Gets the controller that owns the component, this will always be valid during gameplay but can return null in the editor */
 	template <class T>
 	T* GetController() const
 	{
@@ -35,9 +36,10 @@ public:
 
 	//////////////////////////////////////////////////////////////////////////////
 	// Controller accessors
-	// Usable for any type of AController owner
+	// Usable for any type of AController owner, only valid if called during gameplay
 	//////////////////////////////////////////////////////////////////////////////
 
+	/** Returns the pawn that is currently possessed by the owning controller, will often return null */
 	template <class T>
 	T* GetPawn() const
 	{
@@ -45,12 +47,15 @@ public:
 		return Cast<T>(GetControllerChecked<AController>()->GetPawn());
 	}
 
+	/** Returns the actor that is serving as the current view target for the owning controller */
 	template <class T>
 	T* GetViewTarget() const
 	{
+		static_assert(TPointerIsConvertibleFromTo<T, AActor>::Value, "'T' template parameter to GetViewTarget must be derived from APawn");
 		return Cast<T>(GetControllerChecked<AController>()->GetViewTarget());
 	}
 
+	/** If this controller is possessing a pawn return the pawn, if not return the view target */
 	template <class T>
 	T* GetPawnOrViewTarget() const
 	{
@@ -64,6 +69,7 @@ public:
 		}
 	}
 
+	/** Returns the player state attached to this controller if there is one */
 	template <class T>
 	T* GetPlayerState() const
 	{
@@ -71,14 +77,10 @@ public:
 		return GetControllerChecked<AController>()->GetPlayerState<T>();
 	}
 
-	template <class T>
-	T* GetGameInstance() const
-	{
-		static_assert(TPointerIsConvertibleFromTo<T, UGameInstance>::Value, "'T' template parameter to GetGameInstance must be derived from UGameInstance");
-		return GetControllerChecked<AController>()->GetGameInstance<T>();
-	}
-
+	/** Returns true if the owning controller is considered to be local */
 	bool IsLocalController() const;
+
+	/** Returns the point of view for either a player or controlled pawn */
 	void GetPlayerViewPoint(FVector& Location, FRotator& Rotation) const;
 
 	//////////////////////////////////////////////////////////////////////////////
diff --git a/Engine/Plugins/Experimental/ModularGameplay/Source/ModularGameplay/Public/Components/GameFrameworkComponent.h b/Engine/Plugins/Experimental/ModularGameplay/Source/ModularGameplay/Public/Components/GameFrameworkComponent.h
--- a/Engine/Plugins/Experimental/ModularGameplay/Source/ModularGameplay/Public/Components/GameFrameworkComponent.h
+++ b/Engine/Plugins/Experimental/ModularGameplay/Source/ModularGameplay/Public/Components/GameFrameworkComponent.h
@@ -8,7 +8,7 @@
 #include "GameFrameworkComponent.generated.h"
 
 /**
- * GameFrameworkComponent is a base class for a actor components made for the basic game framework classes.
+ * GameFrameworkComponent is a base class for actor components made for the basic game framework classes.
  */
 UCLASS(Blueprintable, BlueprintType, HideCategories=(Trigger, PhysicsVolume))
 class MODULARGAMEPLAY_API UGameFrameworkComponent : public UActorComponent
@@ -19,13 +19,13 @@ public:
 	
 	UGameFrameworkComponent(const FObjectInitializer& ObjectInitializer);
 
+	/** Gets the game instance this component is a part of, this will return null if not called during normal gameplay */
 	template <class T>
 	T* GetGameInstance() const
 	{
 		static_assert(TPointerIsConvertibleFromTo<T, UGameInstance>::Value, "'T' template parameter to GetGameInstance must be derived from UGameInstance");
 		AActor* Owner = GetOwner();
-		check(Owner);
-		return Owner->GetGameInstance<T>();
+		return Owner ? Owner->GetGameInstance<T>() : nullptr;
 	}
 
 	template <class T>
diff --git a/Engine/Plugins/Experimental/ModularGameplay/Source/ModularGameplay/Public/Components/GameStateComponent.h b/Engine/Plugins/Experimental/ModularGameplay/Source/ModularGameplay/Public/Components/GameStateComponent.h
--- a/Engine/Plugins/Experimental/ModularGameplay/Source/ModularGameplay/Public/Components/GameStateComponent.h
+++ b/Engine/Plugins/Experimental/ModularGameplay/Source/ModularGameplay/Public/Components/GameStateComponent.h
@@ -8,7 +8,7 @@
 #include "GameStateComponent.generated.h"
 
 /**
- * GameStateComponent is an actor component made for AGameState and receives GameState events.
+ * GameStateComponent is an actor component made for AGameStateBase and receives GameState events.
  */
 UCLASS()
 class MODULARGAMEPLAY_API UGameStateComponent : public UGameFrameworkComponent
@@ -19,29 +19,31 @@ public:
 	
 	UGameStateComponent(const FObjectInitializer& ObjectInitializer);
 
+	/** Gets the game state that owns the component, this will always be valid during gameplay but can return null in the editor */
 	template <class T>
 	T* GetGameState() const
 	{
-		static_assert(TPointerIsConvertibleFromTo<T, AGameState>::Value, "'T' template parameter to GetGameState must be derived from AGameState");
+		static_assert(TPointerIsConvertibleFromTo<T, AGameStateBase>::Value, "'T' template parameter to GetGameState must be derived from AGameStateBase");
 		return Cast<T>(GetOwner());
 	}
 
 	template <class T>
 	T* GetGameStateChecked() const
 	{
-		static_assert(TPointerIsConvertibleFromTo<T, AGameState>::Value, "'T' template parameter to GetGameStateChecked must be derived from AGameState");
+		static_assert(TPointerIsConvertibleFromTo<T, AGameStateBase>::Value, "'T' template parameter to GetGameStateChecked must be derived from AGameStateBase");
 		return CastChecked<T>(GetOwner());
 	}
 
 	//////////////////////////////////////////////////////////////////////////////
-	// GameState accessors
+	// GameState accessors, only valid if called during gameplay
 	//////////////////////////////////////////////////////////////////////////////
 
+	/** Gets the game mode that owns this component, this will always return null on the client */
 	template <class T>
 	T* GetGameMode() const
 	{
-		static_assert(TPointerIsConvertibleFromTo<T, AGameMode>::Value, "'T' template parameter to GetGameMode must be derived from AGameMode");
-		return Cast<T>(GetGameStateChecked<AGameState>()->AuthorityGameMode);
+		static_assert(TPointerIsConvertibleFromTo<T, AGameModeBase>::Value, "'T' template parameter to GetGameMode must be derived from AGameModeBase");
+		return Cast<T>(GetGameStateChecked<AGameStateBase>()->AuthorityGameMode);
 	}
 
 public:
@@ -50,5 +52,6 @@ public:
 	// GameState events
 	//////////////////////////////////////////////////////////////////////////////
 
+	/** Called when gameplay has fully started */
 	virtual void HandleMatchHasStarted() {}
 };
\ No newline at end of file
diff --git a/Engine/Plugins/Experimental/ModularGameplay/Source/ModularGameplay/Public/Components/PawnComponent.h b/Engine/Plugins/Experimental/ModularGameplay/Source/ModularGameplay/Public/Components/PawnComponent.h
--- a/Engine/Plugins/Experimental/ModularGameplay/Source/ModularGameplay/Public/Components/PawnComponent.h
+++ b/Engine/Plugins/Experimental/ModularGameplay/Source/ModularGameplay/Public/Components/PawnComponent.h
@@ -18,6 +18,7 @@ public:
 	
 	UPawnComponent(const FObjectInitializer& ObjectInitializer);
 
+	/** Gets the pawn that owns the component, this will always be valid during gameplay but can return null in the editor */
 	template <class T>
 	T* GetPawn() const
 	{
@@ -33,9 +34,10 @@ public:
 	}
 
 	//////////////////////////////////////////////////////////////////////////////
-	// Pawn accessors
+	// Pawn accessors, only valid if called during gameplay
 	//////////////////////////////////////////////////////////////////////////////
 
+	/** Gets the player state that owns the component, this can return null on clients for player pawns that are still being replicated */
 	template <class T>
 	T* GetPlayerState() const
 	{
@@ -43,6 +45,7 @@ public:
 		return GetPawnChecked<APawn>()->GetPlayerState<T>();
 	}
 
+	/** Gets the controller that owns the component, this will usually be null on clients */
 	template <class T>
 	T* GetController() const
 	{
@@ -50,10 +53,4 @@ public:
 		return GetPawnChecked<APawn>()->GetController<T>();
 	}
 
-public:
-
-	//////////////////////////////////////////////////////////////////////////////
-	// Pawn events
-	//////////////////////////////////////////////////////////////////////////////
-
 };
\ No newline at end of file
diff --git a/Engine/Plugins/Experimental/ModularGameplay/Source/ModularGameplay/Public/Components/PlayerStateComponent.h b/Engine/Plugins/Experimental/ModularGameplay/Source/ModularGameplay/Public/Components/PlayerStateComponent.h
--- a/Engine/Plugins/Experimental/ModularGameplay/Source/ModularGameplay/Public/Components/PlayerStateComponent.h
+++ b/Engine/Plugins/Experimental/ModularGameplay/Source/ModularGameplay/Public/Components/PlayerStateComponent.h
@@ -18,6 +18,7 @@ public:
 	
 	UPlayerStateComponent(const FObjectInitializer& ObjectInitializer);
 
+	/** Gets the player state that owns the component, this will always be valid during gameplay but can return null in the editor */
 	template <class T>
 	T* GetPlayerState() const
 	{
@@ -33,9 +34,10 @@ public:
 	}
 
 	//////////////////////////////////////////////////////////////////////////////
-	// PlayerState accessors
+	// PlayerState accessors, only valid if called during gameplay
 	//////////////////////////////////////////////////////////////////////////////
 
+	/** Gets the pawn that owns the component, this will return null if the pawn does not exist or is still replicating */
 	template <class T>
 	T* GetPawn() const
 	{
@@ -49,7 +51,10 @@ public:
 	// PlayerState events
 	//////////////////////////////////////////////////////////////////////////////
 
+	/** Reset properties to default state */
 	virtual void Reset() {}
+
+	/** Copy properties which need to be saved in inactive PlayerState */
 	virtual void CopyProperties(UPlayerStateComponent* TargetPlayerStateComponent) {}
 
 };
\ No newline at end of file
