Add arbitrary callback support to GameFrameworkComponentManager to replace the ActorExtensionManager
Default extension events are emitted at registration time and games can broadcast their own

diff --git a/Engine/Plugins/Experimental/ModularGameplay/Source/ModularGameplay/Private/Components/GameFrameworkComponentManager.cpp b/Engine/Plugins/Experimental/ModularGameplay/Source/ModularGameplay/Private/Components/GameFrameworkComponentManager.cpp
--- a/Engine/Plugins/Experimental/ModularGameplay/Source/ModularGameplay/Private/Components/GameFrameworkComponentManager.cpp
+++ b/Engine/Plugins/Experimental/ModularGameplay/Source/ModularGameplay/Private/Components/GameFrameworkComponentManager.cpp
@@ -21,7 +21,14 @@ FComponentRequestHandle::~FComponentRequestHandle()
 	UGameFrameworkComponentManager* LocalManager = OwningManager.Get();
 	if (LocalManager)
 	{
-		LocalManager->RemoveComponentRequest(ReceiverClass, ComponentClass);
+		if (ComponentClass.Get())
+		{
+			LocalManager->RemoveComponentRequest(ReceiverClass, ComponentClass);
+		}
+		if (ExtensionHandle.IsValid())
+		{
+			LocalManager->RemoveExtensionHandler(ReceiverClass, ExtensionHandle);
+		}
 	}
 }
 
@@ -30,6 +37,12 @@ bool FComponentRequestHandle::IsValid() const
 	return OwningManager.IsValid();
 }
 
+FName UGameFrameworkComponentManager::NAME_ReceiverAdded = FName("ReceiverAdded");
+FName UGameFrameworkComponentManager::NAME_ReceiverRemoved = FName("ReceiverRemoved");
+FName UGameFrameworkComponentManager::NAME_ExtensionAdded = FName("ExtensionAdded");
+FName UGameFrameworkComponentManager::NAME_ExtensionRemoved = FName("ExtensionRemoved");
+FName UGameFrameworkComponentManager::NAME_GameActorReady = FName("GameActorReady");
+
 void UGameFrameworkComponentManager::AddReferencedObjects(UObject* InThis, FReferenceCollector& Collector)
 {
 	if (UGameFrameworkComponentManager* GFCM = Cast<UGameFrameworkComponentManager>(InThis))
@@ -124,6 +137,14 @@ void UGameFrameworkComponentManager::AddReceiverInternal(AActor* Receiver)
 				}
 			}
 		}
+
+		if (FExtensionHandlerEvent* HandlerEvent = ReceiverClassToEventMap.Find(ReceiverClassPath))
+		{
+			for (const TPair<FDelegateHandle, FExtensionHandlerDelegate>& Pair : *HandlerEvent)
+			{
+				Pair.Value.Execute(Receiver, NAME_ReceiverAdded);
+			}
+		}
 	}
 }
 
@@ -164,6 +185,8 @@ void UGameFrameworkComponentManager::RemoveReceiverInternal(AActor* Receiver)
 	{
 		DestroyInstancedComponent(Component);
 	}
+
+	SendExtensionEventInternal(Receiver, NAME_ReceiverRemoved);
 }
 
 TSharedPtr<FComponentRequestHandle> UGameFrameworkComponentManager::AddComponentRequest(const TSoftClassPtr<AActor>& ReceiverClass, TSubclassOf<UActorComponent> ComponentClass)
@@ -275,6 +298,122 @@ void UGameFrameworkComponentManager::RemoveComponentRequest(const TSoftClassPtr<
 	}
 }
 
+TSharedPtr<FComponentRequestHandle> UGameFrameworkComponentManager::AddExtensionHandler(const TSoftClassPtr<AActor>& ReceiverClass, FExtensionHandlerDelegate ExtensionHandler)
+{
+	// You must have a target and bound handler. The target cannot be AActor, that is too broad and would be bad for performance.
+	if (!ensure(!ReceiverClass.IsNull()) || !ensure(ExtensionHandler.IsBound()) || !ensure(ReceiverClass.ToString() != TEXT("/Script/Engine.Actor")))
+	{
+		return nullptr;
+	}
+
+	FComponentRequestReceiverClassPath ReceiverClassPath(ReceiverClass);
+	FExtensionHandlerEvent& HandlerEvent = ReceiverClassToEventMap.FindOrAdd(ReceiverClassPath);
+
+	// This is a fake multicast delegate using a map
+	FDelegateHandle DelegateHandle(FDelegateHandle::EGenerateNewHandleType::GenerateNewHandle);
+	HandlerEvent.Add(DelegateHandle, ExtensionHandler);
+
+	if (UClass* ReceiverClassPtr = ReceiverClass.Get())
+	{
+		UGameInstance* LocalGameInstance = GetGameInstance();
+		if (ensure(LocalGameInstance))
+		{
+			UWorld* LocalWorld = LocalGameInstance->GetWorld();
+			if (ensure(LocalWorld))
+			{
+				for (TActorIterator<AActor> ActorIt(LocalWorld, ReceiverClassPtr); ActorIt; ++ActorIt)
+				{
+					if (ActorIt->IsActorInitialized())
+					{
+						ExtensionHandler.Execute(*ActorIt, NAME_ExtensionAdded);
+					}
+				}
+			}
+		}
+	}
+	else
+	{
+		// Actor class is not in memory, there will be no actor instances
+	}
+
+	return MakeShared<FComponentRequestHandle>(this, ReceiverClass, DelegateHandle);
+}
+
+void UGameFrameworkComponentManager::RemoveExtensionHandler(const TSoftClassPtr<AActor>& ReceiverClass, FDelegateHandle DelegateHandle)
+{
+	FComponentRequestReceiverClassPath ReceiverClassPath(ReceiverClass);
+
+	if (FExtensionHandlerEvent* HandlerEvent = ReceiverClassToEventMap.Find(ReceiverClassPath))
+	{
+		FExtensionHandlerDelegate* HandlerDelegate = HandlerEvent->Find(DelegateHandle);
+		if (ensure(HandlerDelegate))
+		{
+			// Call it once on unregister
+			if (UClass* ReceiverClassPtr = ReceiverClass.Get())
+			{
+				UGameInstance* LocalGameInstance = GetGameInstance();
+				if (ensure(LocalGameInstance))
+				{
+					UWorld* LocalWorld = LocalGameInstance->GetWorld();
+					if (ensure(LocalWorld))
+					{
+						for (TActorIterator<AActor> ActorIt(LocalWorld, ReceiverClassPtr); ActorIt; ++ActorIt)
+						{
+							if (ActorIt->IsActorInitialized())
+							{
+								HandlerDelegate->Execute(*ActorIt, NAME_ExtensionRemoved);
+							}
+						}
+					}
+				}
+			}
+			else
+			{
+				// Actor class is not in memory, there will be no actor instances
+			}
+
+			HandlerEvent->Remove(DelegateHandle);
+
+			if (HandlerEvent->IsEmpty())
+			{
+				ReceiverClassToEventMap.Remove(ReceiverClassPath);
+			}
+		}
+	}
+}
+
+void UGameFrameworkComponentManager::SendExtensionEvent(AActor* Receiver, FName EventName, bool bOnlyInGameWorlds)
+{
+	if (Receiver != nullptr)
+	{
+		if (bOnlyInGameWorlds)
+		{
+			UWorld* ReceiverWorld = Receiver->GetWorld();
+			if ((ReceiverWorld == nullptr) || !ReceiverWorld->IsGameWorld() || ReceiverWorld->IsPreviewWorld())
+			{
+				return;
+			}
+		}
+
+		SendExtensionEventInternal(Receiver, EventName);
+	}
+}
+
+void UGameFrameworkComponentManager::SendExtensionEventInternal(AActor* Receiver, const FName& EventName)
+{
+	for (UClass* Class = Receiver->GetClass(); Class && Class != AActor::StaticClass(); Class = Class->GetSuperClass())
+	{
+		FComponentRequestReceiverClassPath ReceiverClassPath(Class);
+		if (FExtensionHandlerEvent* HandlerEvent = ReceiverClassToEventMap.Find(ReceiverClassPath))
+		{
+			for (const TPair<FDelegateHandle, FExtensionHandlerDelegate>& Pair : *HandlerEvent)
+			{
+				Pair.Value.Execute(Receiver, EventName);
+			}
+		}
+	}
+}
+
 void UGameFrameworkComponentManager::CreateComponentOnInstance(AActor* ActorInstance, TSubclassOf<UActorComponent> ComponentClass)
 {
 	check(ActorInstance);
@@ -341,4 +480,29 @@ void UGameFrameworkComponentManager::RemoveGameFrameworkComponentReceiver(AActor
 			GFCM->RemoveReceiverInternal(Receiver);
 		}
 	}
-}
\ No newline at end of file
+}
+
+void UGameFrameworkComponentManager::SendGameFrameworkComponentExtensionEvent(AActor* Receiver, const FName& EventName, bool bOnlyInGameWorlds)
+{
+	if (Receiver != nullptr)
+	{
+		if (bOnlyInGameWorlds)
+		{
+			UWorld* ReceiverWorld = Receiver->GetWorld();
+			if ((ReceiverWorld != nullptr) && ReceiverWorld->IsGameWorld() && !ReceiverWorld->IsPreviewWorld())
+			{
+				if (UGameFrameworkComponentManager* GFCM = UGameInstance::GetSubsystem<UGameFrameworkComponentManager>(ReceiverWorld->GetGameInstance()))
+				{
+					GFCM->SendExtensionEvent(Receiver, EventName);
+				}
+			}
+		}
+		else
+		{
+			if (UGameFrameworkComponentManager* GFCM = UGameInstance::GetSubsystem<UGameFrameworkComponentManager>(Receiver->GetGameInstance()))
+			{
+				GFCM->SendExtensionEvent(Receiver, EventName);
+			}
+		}
+	}
+}
diff --git a/Engine/Plugins/Experimental/ModularGameplay/Source/ModularGameplay/Public/Components/GameFrameworkComponentManager.h b/Engine/Plugins/Experimental/ModularGameplay/Source/ModularGameplay/Public/Components/GameFrameworkComponentManager.h
--- a/Engine/Plugins/Experimental/ModularGameplay/Source/ModularGameplay/Public/Components/GameFrameworkComponentManager.h
+++ b/Engine/Plugins/Experimental/ModularGameplay/Source/ModularGameplay/Public/Components/GameFrameworkComponentManager.h
@@ -16,7 +16,10 @@ class AActor;
 class UActorComponent;
 class UGameFrameworkComponent;
 
-/** A handle for a request to put components on an actor class */
+/** 
+ * A handle for a request to put components or call a delegate for an extensible actor class.
+ * When this handle is destroyed, it will remove the associated request from the system.
+ */
 struct FComponentRequestHandle
 {
 	FComponentRequestHandle(const TWeakObjectPtr<UGameFrameworkComponentManager>& InOwningManager, const TSoftClassPtr<AActor>& InReceiverClass, const TSubclassOf<UActorComponent>& InComponentClass)
@@ -25,6 +28,12 @@ struct FComponentRequestHandle
 		, ComponentClass(InComponentClass)
 	{}
 
+	FComponentRequestHandle(const TWeakObjectPtr<UGameFrameworkComponentManager>& InOwningManager, const TSoftClassPtr<AActor>& InReceiverClass, FDelegateHandle InExtensionHandle)
+		: OwningManager(InOwningManager)
+		, ReceiverClass(InReceiverClass)
+		, ExtensionHandle(InExtensionHandle)
+	{}
+
 	MODULARGAMEPLAY_API ~FComponentRequestHandle();
 
 	/** Returns true if the manager that this request is for still exists */
@@ -39,6 +48,9 @@ private:
 
 	/** The class of component to put on actors */
 	TSubclassOf<UActorComponent> ComponentClass;
+
+	/** A handle to an extension delegate to run */
+	FDelegateHandle ExtensionHandle;
 };
 
 /** 
@@ -46,6 +58,7 @@ private:
  *
  * A manager to handle putting components on actors as they come and go.
  * Put in a request to instantiate components of a given class on actors of a given class and they will automatically be made for them as the actors are spawned.
+ * Submit delegate handlers to listen for actors of a given class. Those handlers will automatically run when actors of a given class or registered as receivers or game events are sent.
  * Actors must opt-in to this behavior by calling AddReceiver/RemoveReceiver for themselves when they are ready to receive the components and when they want to remove them.
  * Any actors that are in memory when a request is made will automatically get the components, and any in memory when a request is removed will lose the components immediately.
  * Requests are reference counted, so if multiple requests are made for the same actor class and component class, only one component will be added and that component wont be removed until all requests are removed.
@@ -55,7 +68,13 @@ class UGameFrameworkComponentManager : public UGameInstanceSubsystem
 {
 	GENERATED_BODY()
 
+	/** Using a fake multicast delegate so order can be kept consistent */
+	DECLARE_DELEGATE_TwoParams(FExtensionHandlerDelegateInternal, AActor*, FName);
+	using FExtensionHandlerEvent = TMap<FDelegateHandle, FExtensionHandlerDelegateInternal>;
 public:
+	/** Delegate types for extension handlers */
+	using FExtensionHandlerDelegate = FExtensionHandlerDelegateInternal;
+
 	void InitGameFrameworkComponentManager() {}
 
 	/** Adds an actor as a receiver for components. If it passes the actorclass filter on requests it will get the components. */
@@ -75,6 +94,40 @@ public:
 	/** Adds a request to instantiate components on actors of the given classes. Returns a handle that will keep the request "alive" until it is destructed, at which point the request is removed. */
 	MODULARGAMEPLAY_API TSharedPtr<FComponentRequestHandle> AddComponentRequest(const TSoftClassPtr<AActor>& ReceiverClass, TSubclassOf<UActorComponent> ComponentClass);
 
+
+	//////////////////////////////////////////////////////////////////////////////////////////////
+	// The extension system allows registering for arbitrary event callbacks on receiver actors.
+	// These are the default events but games can define, send, and listen for their own.
+
+	/** AddReceiver was called for a registered class and components were added, called early in initialization */
+	static MODULARGAMEPLAY_API FName NAME_ReceiverAdded;
+
+	/** RemoveReceiver was called for a registered class and components were removed, normally called from EndPlay */
+	static MODULARGAMEPLAY_API FName NAME_ReceiverRemoved;
+
+	/** A new extension handler was added */
+	static MODULARGAMEPLAY_API FName NAME_ExtensionAdded;
+
+	/** An extension handler was removed by a freed request handle */
+	static MODULARGAMEPLAY_API FName NAME_ExtensionRemoved;
+
+	/** 
+	 * Game-specific event indicating an actor is mostly initialized and ready for extension. 
+	 * All extensible games are expected to send this event at the appropriate actor-specific point, as plugins may be listening for it.
+	 */
+	static MODULARGAMEPLAY_API FName NAME_GameActorReady;
+
+
+	/** Adds an extension handler to run on actors of the given class. Returns a handle that will keep the handler "alive" until it is destructed, at which point the delegate is removed */
+	MODULARGAMEPLAY_API TSharedPtr<FComponentRequestHandle> AddExtensionHandler(const TSoftClassPtr<AActor>& ReceiverClass, FExtensionHandlerDelegate ExtensionHandler);
+
+	/** Sends an arbitrary extension event that can be listened for by other systems */
+	UFUNCTION(BlueprintCallable, Category = "Gameplay", meta = (DefaultToSelf = "Receiver", AdvancedDisplay = 1))
+	MODULARGAMEPLAY_API void SendExtensionEvent(AActor* Receiver, FName EventName, bool bOnlyInGameWorlds = true);
+
+	/** Sends an arbitrary extension event that can be listened for by other systems */
+	static MODULARGAMEPLAY_API void SendGameFrameworkComponentExtensionEvent(AActor* Receiver, const FName& EventName, bool bOnlyInGameWorlds = true);
+
 	static void AddReferencedObjects(UObject* InThis, class FReferenceCollector& Collector);
 
 #if !UE_BUILD_SHIPPING
@@ -84,10 +137,14 @@ public:
 private:
 	void AddReceiverInternal(AActor* Receiver);
 	void RemoveReceiverInternal(AActor* Receiver);
+	void SendExtensionEventInternal(AActor* Receiver, const FName& EventName);
 
 	/** Called by FComponentRequestHandle's destructor to remove a request for components to be created. */
 	void RemoveComponentRequest(const TSoftClassPtr<AActor>& ReceiverClass, TSubclassOf<UActorComponent> ComponentClass);
 
+	/** Called by FComponentRequestHandle's destructor to remove a handler from the system. */
+	void RemoveExtensionHandler(const TSoftClassPtr<AActor>& ReceiverClass, FDelegateHandle DelegateHandle);
+
 	/** Creates an instance of a component on an actor */
 	void CreateComponentOnInstance(AActor* ActorInstance, TSubclassOf<UActorComponent> ComponentClass);
 
@@ -177,9 +234,12 @@ private:
 	/** A map of component classes to instances of that component class made by this component manager */
 	TMap<UClass*, TSet<FObjectKey>> ComponentClassToComponentInstanceMap;
 
-	/** A map actor classes to component classes that should be made for that class. */
+	/** A map of actor classes to component classes that should be made for that class. */
 	TMap<FComponentRequestReceiverClassPath, TSet<UClass*>> ReceiverClassToComponentClassMap;
 
+	/** A map of actor classes to delegate handlers that should be executed for actors of that class. */
+	TMap<FComponentRequestReceiverClassPath, FExtensionHandlerEvent> ReceiverClassToEventMap;
+
 #if WITH_EDITORONLY_DATA
 	/** Editor-only set to validate that component requests are only being added for actors that call AddReceiver and RemoveReceiver */
 	UPROPERTY(transient)
