Prevent an ensure in PIE during level teardown if an action drops a FComponentRequestHandle after some of the worlds have already been torn down (e.g., with game features in multiplayer PIE)

diff --git a/Engine/Plugins/Experimental/ModularGameplay/Source/ModularGameplay/Private/Components/GameFrameworkComponentManager.cpp b/Engine/Plugins/Experimental/ModularGameplay/Source/ModularGameplay/Private/Components/GameFrameworkComponentManager.cpp
--- a/Engine/Plugins/Experimental/ModularGameplay/Source/ModularGameplay/Private/Components/GameFrameworkComponentManager.cpp
+++ b/Engine/Plugins/Experimental/ModularGameplay/Source/ModularGameplay/Private/Components/GameFrameworkComponentManager.cpp
@@ -355,7 +355,8 @@ void UGameFrameworkComponentManager::RemoveExtensionHandler(const TSoftClassPtr<
 				if (ensure(LocalGameInstance))
 				{
 					UWorld* LocalWorld = LocalGameInstance->GetWorld();
-					if (ensure(LocalWorld))
+					ensure(GIsEditor || (LocalWorld != nullptr));
+					if (LocalWorld)
 					{
 						for (TActorIterator<AActor> ActorIt(LocalWorld, ReceiverClassPtr); ActorIt; ++ActorIt)
 						{
