GameFrameworkComponentManager::CreateComponentOnInstance : Do a SetupAttachment() on the RootComponent if the ActorComponent added is a SceneComponent.
#tests PIE
#rnx

diff --git a/Engine/Plugins/Experimental/ModularGameplay/Source/ModularGameplay/Private/Components/GameFrameworkComponentManager.cpp b/Engine/Plugins/Experimental/ModularGameplay/Source/ModularGameplay/Private/Components/GameFrameworkComponentManager.cpp
--- a/Engine/Plugins/Experimental/ModularGameplay/Source/ModularGameplay/Private/Components/GameFrameworkComponentManager.cpp
+++ b/Engine/Plugins/Experimental/ModularGameplay/Source/ModularGameplay/Private/Components/GameFrameworkComponentManager.cpp
@@ -426,6 +426,11 @@ void UGameFrameworkComponentManager::CreateComponentOnInstance(AActor* ActorInst
 		TSet<FObjectKey>& ComponentInstances = ComponentClassToComponentInstanceMap.FindOrAdd(*ComponentClass);
 		ComponentInstances.Add(NewComp);
 
+		if (USceneComponent* NewSceneComp = Cast<USceneComponent>(NewComp))
+		{
+			NewSceneComp->SetupAttachment(ActorInstance->GetRootComponent());
+		}
+
 		NewComp->RegisterComponent();
 	}
 }
