Modular Game Framework - Allow abstract base classes as the 'receiver class' be used without ensuring that they implement AddReceiver/RemoveReceiver.

diff --git a/Engine/Plugins/Experimental/ModularGameplay/Source/ModularGameplay/Private/Components/GameFrameworkComponentManager.cpp b/Engine/Plugins/Experimental/ModularGameplay/Source/ModularGameplay/Private/Components/GameFrameworkComponentManager.cpp
--- a/Engine/Plugins/Experimental/ModularGameplay/Source/ModularGameplay/Private/Components/GameFrameworkComponentManager.cpp
+++ b/Engine/Plugins/Experimental/ModularGameplay/Source/ModularGameplay/Private/Components/GameFrameworkComponentManager.cpp
@@ -224,7 +224,10 @@ TSharedPtr<FComponentRequestHandle> UGameFrameworkComponentManager::AddComponent
 						if (ActorIt->IsActorInitialized())
 						{
 #if WITH_EDITOR
-							ensureMsgf(AllReceivers.Contains(*ActorIt), TEXT("You may not add a component request for an actor class that does not call AddReceiver/RemoveReceiver in code! Class:%s"), *GetPathNameSafe(ReceiverClassPtr));
+							if (!ReceiverClassPtr->HasAllClassFlags(CLASS_Abstract))
+							{
+								ensureMsgf(AllReceivers.Contains(*ActorIt), TEXT("You may not add a component request for an actor class that does not call AddReceiver/RemoveReceiver in code! Class:%s"), *GetPathNameSafe(ReceiverClassPtr));
+							}
 #endif
 							CreateComponentOnInstance(*ActorIt, ComponentClass);
 						}
