Add new Init State system and optional interface to the GameFrameworkComponentManager that allows registering for changes to the initialization states of components
This is designed to be used for modular games like Lyra where features can be enabled by components attached to pawns
This API is a work in progress and won't be final for a few weeks

diff --git a/Engine/Plugins/Experimental/ModularGameplay/Source/ModularGameplay/ModularGameplay.Build.cs b/Engine/Plugins/Experimental/ModularGameplay/Source/ModularGameplay/ModularGameplay.Build.cs
--- a/Engine/Plugins/Experimental/ModularGameplay/Source/ModularGameplay/ModularGameplay.Build.cs
+++ b/Engine/Plugins/Experimental/ModularGameplay/Source/ModularGameplay/ModularGameplay.Build.cs
@@ -20,6 +20,7 @@ public class ModularGameplay : ModuleRules
 			{
 				"Core",
 				"Engine",
+				"GameplayTags",
 			}
 			);
 			
diff --git a/Engine/Plugins/Experimental/ModularGameplay/Source/ModularGameplay/Private/Components/GameFrameworkComponentManager.cpp b/Engine/Plugins/Experimental/ModularGameplay/Source/ModularGameplay/Private/Components/GameFrameworkComponentManager.cpp
--- a/Engine/Plugins/Experimental/ModularGameplay/Source/ModularGameplay/Private/Components/GameFrameworkComponentManager.cpp
+++ b/Engine/Plugins/Experimental/ModularGameplay/Source/ModularGameplay/Private/Components/GameFrameworkComponentManager.cpp
@@ -98,6 +98,31 @@ void UGameFrameworkComponentManager::DumpGameFrameworkComponentManagers()
 }
 #endif // !UE_BUILD_SHIPPING
 
+void UGameFrameworkComponentManager::Initialize(FSubsystemCollectionBase& Collection)
+{
+	Super::Initialize(Collection);
+
+	CurrentStateChange = INDEX_NONE;
+}
+
+UGameFrameworkComponentManager* UGameFrameworkComponentManager::GetForActor(const AActor* Actor, bool bOnlyGameWorlds)
+{
+	if (Actor)
+	{
+		if (UWorld* ReceiverWorld = Actor->GetWorld())
+		{
+			if (bOnlyGameWorlds && (!ReceiverWorld->IsGameWorld() || ReceiverWorld->IsPreviewWorld()))
+			{
+				return nullptr;
+			}
+
+			return UGameInstance::GetSubsystem<UGameFrameworkComponentManager>(ReceiverWorld->GetGameInstance());
+		}
+	}
+
+	return nullptr;
+}
+
 void UGameFrameworkComponentManager::AddReceiver(AActor* Receiver, bool bAddOnlyInGameWorlds)
 {
 	if (Receiver != nullptr)
@@ -115,7 +140,6 @@ void UGameFrameworkComponentManager::AddReceiver(AActor* Receiver, bool bAddOnly
 	}
 }
 
-
 void UGameFrameworkComponentManager::AddReceiverInternal(AActor* Receiver)
 {
 	checkSlow(Receiver);
@@ -457,61 +481,898 @@ void UGameFrameworkComponentManager::DestroyInstancedComponent(UActorComponent*
 
 void UGameFrameworkComponentManager::AddGameFrameworkComponentReceiver(AActor* Receiver, bool bAddOnlyInGameWorlds)
 {
-	if (Receiver != nullptr)
+	if (UGameFrameworkComponentManager* GFCM = GetForActor(Receiver, bAddOnlyInGameWorlds))
 	{
-		if (bAddOnlyInGameWorlds)
+		GFCM->AddReceiverInternal(Receiver);
+	}
+}
+
+void UGameFrameworkComponentManager::RemoveGameFrameworkComponentReceiver(AActor* Receiver)
+{
+	if (UGameFrameworkComponentManager* GFCM = GetForActor(Receiver, false))
+	{
+		GFCM->RemoveReceiverInternal(Receiver);
+	}
+}
+
+void UGameFrameworkComponentManager::SendGameFrameworkComponentExtensionEvent(AActor* Receiver, const FName& EventName, bool bOnlyInGameWorlds)
+{
+	if (UGameFrameworkComponentManager* GFCM = GetForActor(Receiver, bOnlyInGameWorlds))
+	{
+		GFCM->SendExtensionEvent(Receiver, EventName);
+	}
+}
+
+// TODO Actor feature state prototypes
+
+UGameFrameworkComponentManager::FActorFeatureRegisteredDelegate::FActorFeatureRegisteredDelegate(FActorInitStateChangedDelegate&& InDelegate, FName InFeatureName, FGameplayTag InInitState)
+	: Delegate(InDelegate)
+	, DelegateHandle(FDelegateHandle::EGenerateNewHandleType::GenerateNewHandle)
+	, RequiredFeatureName(InFeatureName)
+	, RequiredInitState(InInitState)
+{
+
+}
+
+UGameFrameworkComponentManager::FActorFeatureRegisteredDelegate::FActorFeatureRegisteredDelegate(FActorInitStateChangedBPDelegate&& InDelegate, FName InFeatureName, FGameplayTag InInitState)
+	: BPDelegate(InDelegate)
+	, DelegateHandle(FDelegateHandle::EGenerateNewHandleType::GenerateNewHandle)
+	, RequiredFeatureName(InFeatureName)
+	, RequiredInitState(InInitState)
+{
+
+}
+
+void UGameFrameworkComponentManager::FActorFeatureRegisteredDelegate::Execute(AActor* OwningActor, FName FeatureName, UObject* Implementer, FGameplayTag FeatureState)
+{
+	if (Delegate.IsBound())
+	{
+		ensure(!BPDelegate.IsBound());
+
+		Delegate.Execute(OwningActor, FeatureName, Implementer, FeatureState);
+	}
+	else if (BPDelegate.IsBound())
+	{
+		BPDelegate.Execute(OwningActor, FeatureName, Implementer, FeatureState);
+	}
+}
+
+void UGameFrameworkComponentManager::RegisterInitState(FGameplayTag NewState, bool bAddBefore, FGameplayTag ExistingState)
+{
+	int32 FoundIndex;
+
+	// TODO ensure or ignore for duplicates?
+	if (InitStateOrder.Contains(NewState))
+	{
+		return;
+	}
+
+	if (!ExistingState.IsValid())
+	{
+		if (bAddBefore)
 		{
-			UWorld* ReceiverWorld = Receiver->GetWorld();
-			if ((ReceiverWorld != nullptr) && ReceiverWorld->IsGameWorld() && !ReceiverWorld->IsPreviewWorld())
+			InitStateOrder.Insert(NewState, 0);
+		}
+		else
+		{
+			InitStateOrder.Add(NewState);
+		}
+	}
+
+	else if (ensureMsgf(InitStateOrder.Find(ExistingState, FoundIndex), TEXT("State %s not found in existing state list!"), *ExistingState.ToString()))
+	{
+		if (!bAddBefore)
+		{
+			FoundIndex++;
+		}
+
+		InitStateOrder.Insert(NewState, FoundIndex);
+	}
+}
+
+bool UGameFrameworkComponentManager::IsInitStateAfterOrEqual(FGameplayTag FeatureState, FGameplayTag RelativeState) const
+{
+	// In theory we could use a tree instead of a simple array, but if two states are related to each other they should be registered relative to each other
+	// And if the two states are completely unrelated, the order doesn't matter because it will never be queried
+
+	if (FeatureState == RelativeState)
+	{
+		return true;
+	}
+
+	bool bFoundSecond = false;
+	for (int32 i = 0; i < InitStateOrder.Num(); i++)
+	{
+		if (InitStateOrder[i] == RelativeState)
+		{
+			bFoundSecond = true;
+		}
+		else if (InitStateOrder[i] == FeatureState)
+		{
+			return bFoundSecond;
+		}
+	}
+
+	return false;
+}
+
+FGameplayTag UGameFrameworkComponentManager::GetInitStateForFeature(AActor* Actor, FName FeatureName) const
+{
+	const FActorFeatureState* FoundState = FindFeatureStateStruct(ActorFeatureMap.Find(FObjectKey(Actor)), FeatureName, FGameplayTag());
+
+	if (FoundState)
+	{
+		return FoundState->CurrentState;
+	}
+
+	return FGameplayTag();
+}
+
+bool UGameFrameworkComponentManager::HasFeatureReachedInitState(AActor* Actor, FName FeatureName, FGameplayTag FeatureState) const
+{
+	FGameplayTag FoundState = GetInitStateForFeature(Actor, FeatureName);
+	return IsInitStateAfterOrEqual(FoundState, FeatureState);
+}
+
+UObject* UGameFrameworkComponentManager::GetImplementerForFeature(AActor* Actor, FName FeatureName, FGameplayTag RequiredState) const
+{
+	const FActorFeatureState* FoundState = FindFeatureStateStruct(ActorFeatureMap.Find(FObjectKey(Actor)), FeatureName, RequiredState);
+
+	if (FoundState)
+	{
+		return FoundState->Implementer.Get();
+	}
+
+	return nullptr;
+}
+
+void UGameFrameworkComponentManager::GetAllFeatureImplementers(TArray<UObject*>& OutImplementers, AActor* Actor, FGameplayTag RequiredState, FName ExcludingFeature /*= NAME_None*/) const
+{
+	const FActorFeatureData* FoundStruct = ActorFeatureMap.Find(FObjectKey(Actor));
+
+	if (FoundStruct)
+	{
+		for (const FActorFeatureState& State : FoundStruct->RegisteredStates)
+		{
+			if (State.FeatureName != ExcludingFeature)
 			{
-				if (UGameFrameworkComponentManager* GFCM = UGameInstance::GetSubsystem<UGameFrameworkComponentManager>(ReceiverWorld->GetGameInstance()))
+				if (!RequiredState.IsValid() || IsInitStateAfterOrEqual(State.CurrentState, RequiredState))
 				{
-					GFCM->AddReceiverInternal(Receiver);
+					if (UObject* Implementer = State.Implementer.Get())
+					{
+						OutImplementers.Add(Implementer);
+					}
 				}
 			}
 		}
-		else
+	}
+}
+
+bool UGameFrameworkComponentManager::HaveAllFeaturesReachedInitState(AActor* Actor, FGameplayTag RequiredState, FName ExcludingFeature) const
+{
+	const FActorFeatureData* FoundStruct = ActorFeatureMap.Find(FObjectKey(Actor));
+
+	if (FoundStruct)
+	{
+		for (const FActorFeatureState& State : FoundStruct->RegisteredStates)
 		{
-			if (UGameFrameworkComponentManager* GFCM = UGameInstance::GetSubsystem<UGameFrameworkComponentManager>(Receiver->GetGameInstance()))
+			if (State.FeatureName != ExcludingFeature)
 			{
-				GFCM->AddReceiverInternal(Receiver);
+				if (!IsInitStateAfterOrEqual(State.CurrentState, RequiredState))
+				{
+					return false;
+				}
 			}
 		}
+
+		// TODO do we want at least one feature to be valid?
+		return true;
 	}
+
+	return false;
 }
 
-void UGameFrameworkComponentManager::RemoveGameFrameworkComponentReceiver(AActor* Receiver)
+bool UGameFrameworkComponentManager::ChangeFeatureInitState(AActor* Actor, FName FeatureName, UObject* Implementer, FGameplayTag FeatureState)
 {
-	if (Receiver != nullptr)
+	if (Actor == nullptr || FeatureName.IsNone() || !FeatureState.IsValid())
 	{
-		if (UGameFrameworkComponentManager* GFCM = UGameInstance::GetSubsystem<UGameFrameworkComponentManager>(Receiver->GetGameInstance()))
+		// TODO Ensure?
+		return false;
+	}
+
+	FActorFeatureData& ActorStruct = FindOrAddActorData(Actor);
+
+	FActorFeatureState* FoundState = nullptr;
+	for (FActorFeatureState& State : ActorStruct.RegisteredStates)
+	{
+		if (State.FeatureName == FeatureName)
 		{
-			GFCM->RemoveReceiverInternal(Receiver);
+			// TODO what if it's already in the desired state?
+			FoundState = &State;
 		}
 	}
+
+	if (!FoundState)
+	{
+		FoundState = &ActorStruct.RegisteredStates.Emplace_GetRef(FeatureName);
+	}
+
+	FoundState->CurrentState = FeatureState;
+	FoundState->Implementer = Implementer;
+
+	ProcessFeatureStateChange(Actor, FoundState);
+
+	return true;
 }
 
-void UGameFrameworkComponentManager::SendGameFrameworkComponentExtensionEvent(AActor* Receiver, const FName& EventName, bool bOnlyInGameWorlds)
+void UGameFrameworkComponentManager::RemoveActorFeatureData(AActor* Actor)
 {
-	if (Receiver != nullptr)
+	// TODO when should we clear up now-invalid actor keys?
+	ActorFeatureMap.Remove(FObjectKey(Actor));
+}
+
+void UGameFrameworkComponentManager::RemoveFeatureImplementer(AActor* Actor, UObject* Implementer)
+{
+	if (Actor == nullptr || Implementer == nullptr)
 	{
-		if (bOnlyInGameWorlds)
+		return;
+	}
+	TWeakObjectPtr<UObject> WeakToRemove(Implementer);
+
+	FActorFeatureData* FoundStruct = ActorFeatureMap.Find(FObjectKey(Actor));
+
+	if (FoundStruct)
+	{
+		for (int32 i = FoundStruct->RegisteredStates.Num() - 1; i >= 0; i--)
 		{
-			UWorld* ReceiverWorld = Receiver->GetWorld();
-			if ((ReceiverWorld != nullptr) && ReceiverWorld->IsGameWorld() && !ReceiverWorld->IsPreviewWorld())
+			// Clear if it matches or is stale
+			UObject* ResolvedObject = FoundStruct->RegisteredStates[i].Implementer.Get();
+
+			if (ResolvedObject == Implementer || (!ResolvedObject && !FoundStruct->RegisteredStates[i].Implementer.IsExplicitlyNull()))
+			{
+				FoundStruct->RegisteredStates.RemoveAt(i);
+			}
+		}
+	}
+}
+
+FDelegateHandle UGameFrameworkComponentManager::RegisterAndCallForActorInitState(AActor* Actor, FName FeatureName, FGameplayTag RequiredState, FActorInitStateChangedDelegate Delegate, bool bCallImmediately)
+{
+	if (ensure(Actor && Delegate.IsBound()))
+	{
+		// We often register delegates before registering states
+		FActorFeatureData& ActorStruct = FindOrAddActorData(Actor);
+
+		FActorFeatureRegisteredDelegate& RegisteredDelegate = ActorStruct.RegisteredDelegates.Emplace_GetRef(MoveTemp(Delegate), FeatureName, RequiredState);
+
+		// Cache handle as delegate could invalidate it
+		FDelegateHandle ReturnHandle = RegisteredDelegate.DelegateHandle;
+
+		if (bCallImmediately)
+		{
+			FActorFeatureRegisteredDelegate DelegateCopy = RegisteredDelegate;
+			CallDelegateForMatchingFeatures(Actor, DelegateCopy);
+		}		
+		
+		return ReturnHandle;
+	}
+
+	return FDelegateHandle();
+}
+
+bool UGameFrameworkComponentManager::RegisterAndCallForActorInitState(AActor* Actor, FName FeatureName, FGameplayTag RequiredState, FActorInitStateChangedBPDelegate Delegate, bool bCallImmediately /*= true*/)
+{
+	if (ensure(Actor && Delegate.IsBound()))
+	{
+		// We often register delegates before registering states
+		FActorFeatureData& ActorStruct = FindOrAddActorData(Actor);
+
+		FActorFeatureRegisteredDelegate& RegisteredDelegate = ActorStruct.RegisteredDelegates.Emplace_GetRef(MoveTemp(Delegate), FeatureName, RequiredState);
+
+		// Cache handle as delegate could invalidate it
+		FDelegateHandle ReturnHandle = RegisteredDelegate.DelegateHandle;
+
+		if (bCallImmediately)
+		{
+			FActorFeatureRegisteredDelegate DelegateCopy = RegisteredDelegate;
+			CallDelegateForMatchingFeatures(Actor, DelegateCopy);
+		}
+
+		return true;
+	}
+
+	return false;
+}
+
+bool UGameFrameworkComponentManager::UnregisterActorInitStateDelegate(AActor* Actor, FDelegateHandle& Handle)
+{
+	if (Actor && Handle.IsValid())
+	{
+		FActorFeatureData* ActorStruct = ActorFeatureMap.Find(FObjectKey(Actor));
+
+		if (ActorStruct)
+		{
+			int32 FoundIndex = GetIndexForRegisteredDelegate(ActorStruct->RegisteredDelegates, Handle);
+
+			if (FoundIndex >= 0)
+			{
+				Handle.Reset();
+				ActorStruct->RegisteredDelegates.RemoveAt(FoundIndex);
+				return true;
+			}
+		}
+	}
+
+	return false;
+}
+
+bool UGameFrameworkComponentManager::UnregisterActorInitStateDelegate(AActor* Actor, FActorInitStateChangedBPDelegate DelegateToRemove)
+{
+	if (Actor && DelegateToRemove.IsBound())
+	{
+		FActorFeatureData* ActorStruct = ActorFeatureMap.Find(FObjectKey(Actor));
+
+		if (ActorStruct)
+		{
+			int32 FoundIndex = GetIndexForRegisteredDelegate(ActorStruct->RegisteredDelegates, DelegateToRemove);
+
+			if (FoundIndex >= 0)
+			{
+				ActorStruct->RegisteredDelegates.RemoveAt(FoundIndex);
+				return true;
+			}
+		}
+	}
+
+	return false;
+}
+
+FDelegateHandle UGameFrameworkComponentManager::RegisterAndCallForClassInitState(const TSoftClassPtr<AActor>& ActorClass, FName FeatureName, FGameplayTag RequiredState, FActorInitStateChangedDelegate Delegate, bool bCallImmediately)
+{
+	if (ensure(!ActorClass.IsNull() && Delegate.IsBound() && !FeatureName.IsNone()))
+	{
+		FComponentRequestReceiverClassPath ReceiverClassPath(ActorClass);
+		TArray<FActorFeatureRegisteredDelegate>& RegisteredDelegates = ClassFeatureChangeDelegates.FindOrAdd(ReceiverClassPath);
+
+		FActorFeatureRegisteredDelegate& RegisteredDelegate = RegisteredDelegates.Emplace_GetRef(MoveTemp(Delegate), FeatureName, RequiredState);
+
+		// Cache handle as delegate could invalidate it
+		FDelegateHandle ReturnHandle = RegisteredDelegate.DelegateHandle;
+
+		if (bCallImmediately)
+		{
+			// If RealClass isn't valid yet, there can't be any instances
+			UClass* RealClass = ActorClass.Get();
+
+			FActorFeatureRegisteredDelegate DelegateCopy = RegisteredDelegate;
+			CallDelegateForMatchingActors(RealClass, DelegateCopy);
+		}
+
+		return ReturnHandle;
+	}
+
+	return FDelegateHandle();
+}
+
+bool UGameFrameworkComponentManager::RegisterAndCallForClassInitState(TSoftClassPtr<AActor> ActorClass, FName FeatureName, FGameplayTag RequiredState, FActorInitStateChangedBPDelegate Delegate, bool bCallImmediately /*= true*/)
+{
+	if (ensure(!ActorClass.IsNull() && Delegate.IsBound() && !FeatureName.IsNone()))
+	{
+		FComponentRequestReceiverClassPath ReceiverClassPath(ActorClass);
+		TArray<FActorFeatureRegisteredDelegate>& RegisteredDelegates = ClassFeatureChangeDelegates.FindOrAdd(ReceiverClassPath);
+
+		FActorFeatureRegisteredDelegate& RegisteredDelegate = RegisteredDelegates.Emplace_GetRef(MoveTemp(Delegate), FeatureName, RequiredState);
+
+		// Cache handle as delegate could invalidate it
+		FDelegateHandle ReturnHandle = RegisteredDelegate.DelegateHandle;
+
+		if (bCallImmediately)
+		{
+			// If RealClass isn't valid yet, there can't be any instances
+			UClass* RealClass = ActorClass.Get();
+
+			FActorFeatureRegisteredDelegate DelegateCopy = RegisteredDelegate;
+			CallDelegateForMatchingActors(RealClass, DelegateCopy);
+		}
+
+		return true;
+	}
+
+	return false;
+}
+
+bool UGameFrameworkComponentManager::UnregisterClassInitStateDelegate(const TSoftClassPtr<AActor>& ActorClass, FDelegateHandle& Handle)
+{
+	if (!ActorClass.IsNull() && Handle.IsValid())
+	{
+		FComponentRequestReceiverClassPath ReceiverClassPath(ActorClass);
+		TArray<FActorFeatureRegisteredDelegate>* RegisteredDelegates = ClassFeatureChangeDelegates.Find(ReceiverClassPath);
+
+		if (RegisteredDelegates)
+		{
+			int32 FoundIndex = GetIndexForRegisteredDelegate(*RegisteredDelegates, Handle);
+
+			if (FoundIndex >= 0)
+			{
+				Handle.Reset();
+				RegisteredDelegates->RemoveAt(FoundIndex);
+				return true;
+			}
+		}
+	}
+
+	return false;
+}
+
+bool UGameFrameworkComponentManager::UnregisterClassInitStateDelegate(TSoftClassPtr<AActor> ActorClass, FActorInitStateChangedBPDelegate DelegateToRemove)
+{
+	if (!ActorClass.IsNull() && DelegateToRemove.IsBound())
+	{
+		FComponentRequestReceiverClassPath ReceiverClassPath(ActorClass);
+		TArray<FActorFeatureRegisteredDelegate>* RegisteredDelegates = ClassFeatureChangeDelegates.Find(ReceiverClassPath);
+
+		if (RegisteredDelegates)
+		{
+			int32 FoundIndex = GetIndexForRegisteredDelegate(*RegisteredDelegates, DelegateToRemove);
+
+			if (FoundIndex >= 0)
+			{
+				RegisteredDelegates->RemoveAt(FoundIndex);
+				return true;
+			}
+		}
+	}
+
+	return false;
+}
+
+
+const UGameFrameworkComponentManager::FActorFeatureState* UGameFrameworkComponentManager::FindFeatureStateStruct(const FActorFeatureData* ActorStruct, FName FeatureName, FGameplayTag RequiredState) const
+{
+	if (ActorStruct)
+	{
+		for (const FActorFeatureState& State : ActorStruct->RegisteredStates)
+		{
+			if (State.FeatureName == FeatureName)
+			{
+				if (!RequiredState.IsValid() || IsInitStateAfterOrEqual(State.CurrentState, RequiredState))
+				{
+					return &State;
+				}
+			}
+		}
+	}
+
+	return nullptr;
+}
+
+void UGameFrameworkComponentManager::ProcessFeatureStateChange(AActor* Actor, const FActorFeatureState* StateChange)
+{
+	StateChangeQueue.Emplace(Actor, *StateChange);
+
+	if (CurrentStateChange == INDEX_NONE)
+	{
+		// Start processing in order
+		CurrentStateChange = 0;
+
+		while (CurrentStateChange < StateChangeQueue.Num())
+		{
+			CallFeatureStateDelegates(StateChangeQueue[CurrentStateChange].Key, StateChangeQueue[CurrentStateChange].Value);
+			CurrentStateChange++;
+		}
+
+		// Done processing, clear it
+		StateChangeQueue.Empty();
+		CurrentStateChange = INDEX_NONE;
+	}
+}
+
+void UGameFrameworkComponentManager::CallFeatureStateDelegates(AActor* Actor, FActorFeatureState StateChange)
+{
+	FActorFeatureData* ActorStruct = ActorFeatureMap.Find(FObjectKey(Actor));
+	int32 DelegateIndex = 0;
+	bool bMoreActorDelegates = true;
+	bool bMoreClassDelegates = true;
+	UClass* ClassToCheck = Actor->GetClass();
+
+	while (ActorStruct)
+	{
+		if (bMoreActorDelegates)
+		{
+			// First check the actor-specific delegates in order
+
+			if (ActorStruct->RegisteredDelegates.IsValidIndex(DelegateIndex))
 			{
-				if (UGameFrameworkComponentManager* GFCM = UGameInstance::GetSubsystem<UGameFrameworkComponentManager>(ReceiverWorld->GetGameInstance()))
+				FActorFeatureRegisteredDelegate& RegisteredDelegate = ActorStruct->RegisteredDelegates[DelegateIndex];
+
+				if ((RegisteredDelegate.RequiredFeatureName.IsNone() || RegisteredDelegate.RequiredFeatureName == StateChange.FeatureName)
+					&& (!RegisteredDelegate.RequiredInitState.IsValid() || IsInitStateAfterOrEqual(StateChange.CurrentState, RegisteredDelegate.RequiredInitState)))
+				{
+					// TODO remove invalid delegates now?
+					RegisteredDelegate.Execute(Actor, StateChange.FeatureName, StateChange.Implementer.Get(), StateChange.CurrentState);
+
+					// That could have invalidated anything
+					ActorStruct = ActorFeatureMap.Find(FObjectKey(Actor));
+				}
+
+				DelegateIndex++;
+			}
+			else
+			{
+				DelegateIndex = 0;
+				bMoreActorDelegates = false;
+			}
+		}
+		else if (ClassToCheck)
+		{
+			// Now check the general class delegates
+
+			FComponentRequestReceiverClassPath ReceiverClassPath(ClassToCheck);
+			TArray<FActorFeatureRegisteredDelegate>* FoundDelegates = ClassFeatureChangeDelegates.Find(ReceiverClassPath);
+
+			if (FoundDelegates && FoundDelegates->IsValidIndex(DelegateIndex))
+			{
+				FActorFeatureRegisteredDelegate& RegisteredDelegate = (*FoundDelegates)[DelegateIndex];
+
+				if ((RegisteredDelegate.RequiredFeatureName.IsNone() || RegisteredDelegate.RequiredFeatureName == StateChange.FeatureName)
+					&& (!RegisteredDelegate.RequiredInitState.IsValid() || IsInitStateAfterOrEqual(StateChange.CurrentState, RegisteredDelegate.RequiredInitState)))
 				{
-					GFCM->SendExtensionEvent(Receiver, EventName);
+					// TODO remove invalid delegates now?
+					RegisteredDelegate.Execute(Actor, StateChange.FeatureName, StateChange.Implementer.Get(), StateChange.CurrentState);
+
+					// That could have invalidated anything
+					ActorStruct = ActorFeatureMap.Find(FObjectKey(Actor));
 				}
+
+				DelegateIndex++;
+			}
+			else
+			{
+				// Start over with parent class
+				DelegateIndex = 0;
+				ClassToCheck = ClassToCheck->GetSuperClass();
 			}
 		}
 		else
 		{
-			if (UGameFrameworkComponentManager* GFCM = UGameInstance::GetSubsystem<UGameFrameworkComponentManager>(Receiver->GetGameInstance()))
+			return;
+		}
+	}
+}
+
+void UGameFrameworkComponentManager::CallDelegateForMatchingFeatures(AActor* Actor, FActorFeatureRegisteredDelegate& RegisteredDelegate)
+{
+	FActorFeatureData* ActorStruct = ActorFeatureMap.Find(FObjectKey(Actor));
+	int32 StateIndex = 0;
+
+	// If feature is specified, just call the one
+	if (!RegisteredDelegate.RequiredFeatureName.IsNone())
+	{
+		const FActorFeatureState* FoundStruct = FindFeatureStateStruct(ActorStruct, RegisteredDelegate.RequiredFeatureName, RegisteredDelegate.RequiredInitState);
+
+		if (FoundStruct)
+		{
+			RegisteredDelegate.Execute(Actor, FoundStruct->FeatureName, FoundStruct->Implementer.Get(), FoundStruct->CurrentState);
+		}
+
+		return;
+	}
+	
+	// If feature is not specified, iterate them all
+	while (ActorStruct && ActorStruct->RegisteredStates.IsValidIndex(StateIndex))
+	{
+		const FActorFeatureState* FoundStruct = &ActorStruct->RegisteredStates[StateIndex];
+
+		if (!RegisteredDelegate.RequiredInitState.IsValid() || IsInitStateAfterOrEqual(FoundStruct->CurrentState, RegisteredDelegate.RequiredInitState))
+		{
+			RegisteredDelegate.Execute(Actor, FoundStruct->FeatureName, FoundStruct->Implementer.Get(), FoundStruct->CurrentState);
+
+			// That could have invalidated anything
+			ActorStruct = ActorFeatureMap.Find(FObjectKey(Actor));
+		}
+
+		StateIndex++;
+	}
+}
+
+void UGameFrameworkComponentManager::CallDelegateForMatchingActors(UClass* ActorClass, FActorFeatureRegisteredDelegate& RegisteredDelegate)
+{
+	TArray<AActor*> MatchingActors;
+
+	if (ActorClass == nullptr)
+	{
+		return;
+	}
+
+	for (TPair<FObjectKey, FActorFeatureData>& Pair : ActorFeatureMap)
+	{
+		UClass* CheckActorClass = Pair.Value.ActorClass.Get();
+
+		if (CheckActorClass->IsChildOf(ActorClass))
+		{
+			AActor* FoundActor = Cast<AActor>(Pair.Key.ResolveObjectPtr());
+			if (FoundActor)
+			{
+				MatchingActors.Add(FoundActor);
+			}
+		}
+	}
+
+	// Iterate actor list before calling any delegates as delegates could change it
+	for (AActor* Actor : MatchingActors)
+	{
+		CallDelegateForMatchingFeatures(Actor, RegisteredDelegate);
+	}
+}
+
+UGameFrameworkComponentManager::FActorFeatureData& UGameFrameworkComponentManager::FindOrAddActorData(AActor* Actor)
+{
+	check(Actor);
+
+	FActorFeatureData& ActorStruct = ActorFeatureMap.FindOrAdd(FObjectKey(Actor));
+	if (!ActorStruct.ActorClass.IsValid())
+	{
+		ActorStruct.ActorClass = Actor->GetClass();
+	}
+	return ActorStruct;
+}
+
+int32 UGameFrameworkComponentManager::GetIndexForRegisteredDelegate(TArray<FActorFeatureRegisteredDelegate>& DelegatesToSearch, FDelegateHandle SearchHandle) const
+{
+	for (int32 i = DelegatesToSearch.Num() - 1; i >= 0; i--)
+	{
+		if (DelegatesToSearch[i].DelegateHandle == SearchHandle)
+		{
+			return i;
+		}
+	}
+
+	return INDEX_NONE;
+}
+
+int32 UGameFrameworkComponentManager::GetIndexForRegisteredDelegate(TArray<FActorFeatureRegisteredDelegate>& DelegatesToSearch, FActorInitStateChangedBPDelegate SearchDelegate) const
+{
+	for (int32 i = DelegatesToSearch.Num() - 1; i >= 0; i--)
+	{
+		if (DelegatesToSearch[i].BPDelegate == SearchDelegate)
+		{
+			return i;
+		}
+	}
+
+	return INDEX_NONE;
+}
+
+// INTERFACE
+
+AActor* IGameFrameworkInitStateInterface::GetOwningActor() const
+{
+	// Removing const because most AActor functions can't handle const
+	AActor* FoundActor = const_cast<AActor*>(Cast<AActor>(this));
+
+	if (!FoundActor)
+	{
+		const UActorComponent* FoundComponent = Cast<UActorComponent>(this);
+		if (FoundComponent)
+		{
+			FoundActor = FoundComponent->GetOwner();
+		}
+	}
+
+	if (ensure(FoundActor))
+	{
+		// Subclasses must implement this if they are not actors or components
+		return FoundActor;
+	}
+
+	return nullptr;
+}
+
+UGameFrameworkComponentManager* IGameFrameworkInitStateInterface::GetComponentManager() const
+{
+	return UGameFrameworkComponentManager::GetForActor(GetOwningActor());
+}
+
+FGameplayTag IGameFrameworkInitStateInterface::GetInitState() const
+{
+	AActor* MyActor = GetOwningActor();
+	UGameFrameworkComponentManager* Manager = UGameFrameworkComponentManager::GetForActor(MyActor);
+	if (Manager)
+	{
+		return Manager->GetInitStateForFeature(MyActor, GetFeatureName());
+	}
+
+	return FGameplayTag();
+}
+
+bool IGameFrameworkInitStateInterface::HasReachedInitState(FGameplayTag DesiredState) const
+{
+	AActor* MyActor = GetOwningActor();
+	UGameFrameworkComponentManager* Manager = UGameFrameworkComponentManager::GetForActor(MyActor);
+	if (Manager)
+	{
+		return Manager->HasFeatureReachedInitState(MyActor, GetFeatureName(), DesiredState);
+	}
+
+	return false;
+}
+
+bool IGameFrameworkInitStateInterface::TryToChangeInitState(FGameplayTag DesiredState)
+{
+	UObject* ThisObject = Cast<UObject>(this);
+	AActor* MyActor = GetOwningActor();
+	UGameFrameworkComponentManager* Manager = UGameFrameworkComponentManager::GetForActor(MyActor);
+	const FName MyFeatureName = GetFeatureName();
+
+	if (!Manager || !ThisObject || !MyActor)
+	{
+		return false;	
+	}
+
+	FGameplayTag CurrentState = Manager->GetInitStateForFeature(MyActor, MyFeatureName);
+	
+	// If we're already in that state, just return
+	if (CurrentState == DesiredState)
+	{
+		return false;
+	}
+
+	if (!CanChangeInitState(Manager, CurrentState, DesiredState))
+	{
+		UE_LOG(LogModularGameplay, Verbose, TEXT("TryToChangeInitState: Cannot transition %s:%s (role %d) from %s to %s"),
+			*MyActor->GetName(), *MyFeatureName.ToString(), MyActor->GetLocalRole(), *CurrentState.ToString(), *DesiredState.ToString());
+		return false;
+	}
+
+	UE_LOG(LogModularGameplay, Verbose, TEXT("TryToChangeInitState: Transitioning %s:%s (role %d) from %s to %s"),
+		*MyActor->GetName(), *MyFeatureName.ToString(), MyActor->GetLocalRole(), *CurrentState.ToString(), *DesiredState.ToString());
+
+	// Perform the local change
+	HandleChangeInitState(Manager, CurrentState, DesiredState);
+
+	// The local change has completed, notify the system to register change and execute callbacks
+	return ensure(Manager->ChangeFeatureInitState(MyActor, MyFeatureName, ThisObject, DesiredState));
+}
+
+FGameplayTag IGameFrameworkInitStateInterface::ContinueInitStateChain(const TArray<FGameplayTag>& InitStateChain)
+{
+	UObject* ThisObject = Cast<UObject>(this);
+	AActor* MyActor = GetOwningActor();
+	UGameFrameworkComponentManager* Manager = UGameFrameworkComponentManager::GetForActor(MyActor);
+	const FName MyFeatureName = GetFeatureName();
+
+	if (!Manager || !ThisObject || !MyActor)
+	{
+		return FGameplayTag();
+	}
+
+	int32 ChainIndex = 0;
+	FGameplayTag CurrentState = Manager->GetInitStateForFeature(MyActor, MyFeatureName);
+
+	// For each state in chain before the last, see if we can transition to the next state
+	while (ChainIndex < InitStateChain.Num() - 1)
+	{
+		if (CurrentState == InitStateChain[ChainIndex])
+		{
+			FGameplayTag DesiredState = InitStateChain[ChainIndex + 1];
+			if (CanChangeInitState(Manager, CurrentState, DesiredState))
 			{
-				GFCM->SendExtensionEvent(Receiver, EventName);
+				UE_LOG(LogModularGameplay, Verbose, TEXT("ContinueInitStateChain: Transitioning %s:%s (role %d) from %s to %s"),
+					*MyActor->GetName(), *MyFeatureName.ToString(), MyActor->GetLocalRole(), *CurrentState.ToString(), *DesiredState.ToString());
+
+				// Perform the local change
+				HandleChangeInitState(Manager, CurrentState, DesiredState);
+
+				// The local change has completed, notify the system to register change and execute callbacks
+				ensure(Manager->ChangeFeatureInitState(MyActor, MyFeatureName, ThisObject, DesiredState));
+
+				// Update state and check again
+				CurrentState = Manager->GetInitStateForFeature(MyActor, MyFeatureName);
 			}
+			else
+			{
+				UE_LOG(LogModularGameplay, Verbose, TEXT("ContinueInitStateChain: Cannot transition %s:%s (role %d) from %s to %s"),
+					*MyActor->GetName(), *MyFeatureName.ToString(), MyActor->GetLocalRole(), *CurrentState.ToString(), *DesiredState.ToString());
+			}
+		}
+
+		ChainIndex++;
+	}
+
+	return CurrentState;
+}
+
+void IGameFrameworkInitStateInterface::CheckDefaultInitializationForImplementers()
+{
+	UObject* ThisObject = Cast<UObject>(this);
+	AActor* MyActor = GetOwningActor();
+	UGameFrameworkComponentManager* Manager = UGameFrameworkComponentManager::GetForActor(MyActor);
+	const FName MyFeatureName = GetFeatureName();
+
+	if (Manager)
+	{
+		TArray<UObject*> Implementers;
+		Manager->GetAllFeatureImplementers(Implementers, MyActor, FGameplayTag(), MyFeatureName);
+
+		for (UObject* Implementer : Implementers)
+		{
+			if (IGameFrameworkInitStateInterface* ImplementerInterface = Cast<IGameFrameworkInitStateInterface>(Implementer))
+			{
+				ImplementerInterface->CheckDefaultInitialization();
+			}
+		}
+	}
+}
+
+void IGameFrameworkInitStateInterface::UnregisterInitStateData()
+{
+	UObject* ThisObject = Cast<UObject>(this);
+	AActor* MyActor = GetOwningActor();
+	UGameFrameworkComponentManager* Manager = UGameFrameworkComponentManager::GetForActor(MyActor);
+
+	if (Manager)
+	{
+		if (ThisObject == MyActor)
+		{
+			// This will clear all the feature states and delegates
+			Manager->RemoveActorFeatureData(MyActor);
+		}
+		else
+		{
+			Manager->RemoveFeatureImplementer(MyActor, ThisObject);
 		}
+
+		if (ActorInitStateChangedHandle.IsValid())
+		{
+			Manager->UnregisterActorInitStateDelegate(MyActor, ActorInitStateChangedHandle);
+		}
+	}
+}
+
+void IGameFrameworkInitStateInterface::BindOnActorInitStateChanged(FName FeatureName, FGameplayTag RequiredState, bool bCallIfReached)
+{
+	UObject* ThisObject = Cast<UObject>(this);
+	AActor* MyActor = GetOwningActor();
+	UGameFrameworkComponentManager* Manager = UGameFrameworkComponentManager::GetForActor(MyActor);
+
+	if (ensure(MyActor && Manager))
+	{
+		// Bind as a weak lambda because this is not a UObject but is guaranteed to be valid as long as ThisObject is
+		FActorInitStateChangedDelegate Delegate = FActorInitStateChangedDelegate::CreateWeakLambda(ThisObject,
+			[this](AActor* OwningActor, FName FeatureName, UObject* Implementer, FGameplayTag FeatureState)
+			{
+				this->OnActorInitStateChanged(OwningActor, FeatureName, Implementer, FeatureState);
+			});
+
+		ActorInitStateChangedHandle = Manager->RegisterAndCallForActorInitState(MyActor, FeatureName, RequiredState, MoveTemp(Delegate), bCallIfReached);
 	}
 }
+
+bool IGameFrameworkInitStateInterface::RegisterAndCallForInitStateChange(FGameplayTag RequiredState, FActorInitStateChangedBPDelegate Delegate, bool bCallImmediately)
+{
+	UObject* ThisObject = Cast<UObject>(this);
+	AActor* MyActor = GetOwningActor();
+	UGameFrameworkComponentManager* Manager = UGameFrameworkComponentManager::GetForActor(MyActor);
+	const FName MyFeatureName = GetFeatureName();
+
+	if (ensure(MyActor && Manager))
+	{
+		return Manager->RegisterAndCallForActorInitState(MyActor, MyFeatureName, RequiredState, Delegate, bCallImmediately);
+	}
+
+	return false;
+}
+
+bool IGameFrameworkInitStateInterface::UnregisterInitStateDelegate(FActorInitStateChangedBPDelegate Delegate)
+{
+	UObject* ThisObject = Cast<UObject>(this);
+	AActor* MyActor = GetOwningActor();
+	UGameFrameworkComponentManager* Manager = UGameFrameworkComponentManager::GetForActor(MyActor);
+
+	if (ensure(MyActor && Manager))
+	{
+		return Manager->UnregisterActorInitStateDelegate(MyActor, Delegate);
+	}
+
+	return false;
+}
diff --git a/Engine/Plugins/Experimental/ModularGameplay/Source/ModularGameplay/Public/Components/GameFrameworkComponentManager.h b/Engine/Plugins/Experimental/ModularGameplay/Source/ModularGameplay/Public/Components/GameFrameworkComponentManager.h
--- a/Engine/Plugins/Experimental/ModularGameplay/Source/ModularGameplay/Public/Components/GameFrameworkComponentManager.h
+++ b/Engine/Plugins/Experimental/ModularGameplay/Source/ModularGameplay/Public/Components/GameFrameworkComponentManager.h
@@ -9,6 +9,7 @@
 #include "UObject/ClassTree.h"
 #include "Templates/SubclassOf.h"
 #include "Subsystems/GameInstanceSubsystem.h"
+#include "GameplayTagContainer.h"
 
 #include "GameFrameworkComponentManager.generated.h"
 
@@ -16,6 +17,12 @@ class AActor;
 class UActorComponent;
 class UGameFrameworkComponent;
 
+/** Native delegate called when an actor feature changes init state */
+DECLARE_DELEGATE_FourParams(FActorInitStateChangedDelegate, AActor* /* OwningActor */, FName /* FeatureName */, UObject* /* Implementer */, FGameplayTag /* FeatureState */);
+
+/** Blueprint delegate called when an actor feature changes init state */
+DECLARE_DYNAMIC_DELEGATE_FourParams(FActorInitStateChangedBPDelegate, AActor*, OwningActor, FName, FeatureName, UObject*, Implementer, FGameplayTag, FeatureState);
+
 /** 
  * A handle for a request to put components or call a delegate for an extensible actor class.
  * When this handle is destroyed, it will remove the associated request from the system.
@@ -63,8 +70,8 @@ private:
  * Any actors that are in memory when a request is made will automatically get the components, and any in memory when a request is removed will lose the components immediately.
  * Requests are reference counted, so if multiple requests are made for the same actor class and component class, only one component will be added and that component wont be removed until all requests are removed.
  */
-UCLASS(MinimalAPI)
-class UGameFrameworkComponentManager : public UGameInstanceSubsystem
+UCLASS()
+class MODULARGAMEPLAY_API UGameFrameworkComponentManager : public UGameInstanceSubsystem
 {
 	GENERATED_BODY()
 
@@ -75,24 +82,27 @@ public:
 	/** Delegate types for extension handlers */
 	using FExtensionHandlerDelegate = FExtensionHandlerDelegateInternal;
 
-	void InitGameFrameworkComponentManager() {}
+	virtual void Initialize(FSubsystemCollectionBase& Collection) override;
+
+	/** Utility to get this manager from an actor, will return null if actor is null or not in a world */
+	static UGameFrameworkComponentManager* GetForActor(const AActor* Actor, bool bOnlyGameWorlds = true);
 
 	/** Adds an actor as a receiver for components. If it passes the actorclass filter on requests it will get the components. */
 	UFUNCTION(BlueprintCallable, Category="Gameplay", meta=(DefaultToSelf="Receiver", AdvancedDisplay=1))
-	MODULARGAMEPLAY_API void AddReceiver(AActor* Receiver, bool bAddOnlyInGameWorlds = true);
+	void AddReceiver(AActor* Receiver, bool bAddOnlyInGameWorlds = true);
 
 	/** Removes an actor as a receiver for components. */
 	UFUNCTION(BlueprintCallable, Category="Gameplay", meta=(DefaultToSelf="Receiver"))
-	MODULARGAMEPLAY_API void RemoveReceiver(AActor* Receiver);
+	void RemoveReceiver(AActor* Receiver);
 
 	/** Adds an actor as a receiver for components (automatically finding the manager for the actor's  game instance). If it passes the actorclass filter on requests it will get the components. */
-	static MODULARGAMEPLAY_API void AddGameFrameworkComponentReceiver(AActor* Receiver, bool bAddOnlyInGameWorlds = true);
+	static void AddGameFrameworkComponentReceiver(AActor* Receiver, bool bAddOnlyInGameWorlds = true);
 
 	/** Removes an actor as a receiver for components (automatically finding the manager for the actor's game instance). */
-	static MODULARGAMEPLAY_API void RemoveGameFrameworkComponentReceiver(AActor* Receiver);
+	static void RemoveGameFrameworkComponentReceiver(AActor* Receiver);
 
 	/** Adds a request to instantiate components on actors of the given classes. Returns a handle that will keep the request "alive" until it is destructed, at which point the request is removed. */
-	MODULARGAMEPLAY_API TSharedPtr<FComponentRequestHandle> AddComponentRequest(const TSoftClassPtr<AActor>& ReceiverClass, TSubclassOf<UActorComponent> ComponentClass);
+	TSharedPtr<FComponentRequestHandle> AddComponentRequest(const TSoftClassPtr<AActor>& ReceiverClass, TSubclassOf<UActorComponent> ComponentClass);
 
 
 	//////////////////////////////////////////////////////////////////////////////////////////////
@@ -100,33 +110,33 @@ public:
 	// These are the default events but games can define, send, and listen for their own.
 
 	/** AddReceiver was called for a registered class and components were added, called early in initialization */
-	static MODULARGAMEPLAY_API FName NAME_ReceiverAdded;
+	static FName NAME_ReceiverAdded;
 
 	/** RemoveReceiver was called for a registered class and components were removed, normally called from EndPlay */
-	static MODULARGAMEPLAY_API FName NAME_ReceiverRemoved;
+	static FName NAME_ReceiverRemoved;
 
 	/** A new extension handler was added */
-	static MODULARGAMEPLAY_API FName NAME_ExtensionAdded;
+	static FName NAME_ExtensionAdded;
 
 	/** An extension handler was removed by a freed request handle */
-	static MODULARGAMEPLAY_API FName NAME_ExtensionRemoved;
+	static FName NAME_ExtensionRemoved;
 
 	/** 
 	 * Game-specific event indicating an actor is mostly initialized and ready for extension. 
 	 * All extensible games are expected to send this event at the appropriate actor-specific point, as plugins may be listening for it.
 	 */
-	static MODULARGAMEPLAY_API FName NAME_GameActorReady;
+	static FName NAME_GameActorReady;
 
 
 	/** Adds an extension handler to run on actors of the given class. Returns a handle that will keep the handler "alive" until it is destructed, at which point the delegate is removed */
-	MODULARGAMEPLAY_API TSharedPtr<FComponentRequestHandle> AddExtensionHandler(const TSoftClassPtr<AActor>& ReceiverClass, FExtensionHandlerDelegate ExtensionHandler);
+	TSharedPtr<FComponentRequestHandle> AddExtensionHandler(const TSoftClassPtr<AActor>& ReceiverClass, FExtensionHandlerDelegate ExtensionHandler);
 
 	/** Sends an arbitrary extension event that can be listened for by other systems */
 	UFUNCTION(BlueprintCallable, Category = "Gameplay", meta = (DefaultToSelf = "Receiver", AdvancedDisplay = 1))
-	MODULARGAMEPLAY_API void SendExtensionEvent(AActor* Receiver, FName EventName, bool bOnlyInGameWorlds = true);
+	void SendExtensionEvent(AActor* Receiver, FName EventName, bool bOnlyInGameWorlds = true);
 
 	/** Sends an arbitrary extension event that can be listened for by other systems */
-	static MODULARGAMEPLAY_API void SendGameFrameworkComponentExtensionEvent(AActor* Receiver, const FName& EventName, bool bOnlyInGameWorlds = true);
+	static void SendGameFrameworkComponentExtensionEvent(AActor* Receiver, const FName& EventName, bool bOnlyInGameWorlds = true);
 
 	static void AddReferencedObjects(UObject* InThis, class FReferenceCollector& Collector);
 
@@ -247,4 +257,267 @@ private:
 #endif // WITH_EDITORONLY_DATA
 
 	friend struct FComponentRequestHandle;
+
+	// TODO work in progress for new actor feature ini state system
+
+
+public:
+
+	/** Adds a new global actor feature state, either before or after an existing one. This will generally be called from game global or game feature initialization */
+	void RegisterInitState(FGameplayTag NewState, bool bAddBefore, FGameplayTag ExistingState);
+
+	/** Returns true if FeatureState comes after the second state (or is equal) */
+	bool IsInitStateAfterOrEqual(FGameplayTag FeatureState, FGameplayTag RelativeState) const;
+
+	/** Returns the earliest state found for the given feature */
+	FGameplayTag GetInitStateForFeature(AActor* Actor, FName FeatureName) const;
+
+	/** Returns true if feature has reached query state or later */
+	bool HasFeatureReachedInitState(AActor* Actor, FName FeatureName, FGameplayTag FeatureState) const;
+
+	/** Returns the object implementing specified feature, filtered by required state if not none */
+	UObject* GetImplementerForFeature(AActor* Actor, FName FeatureName, FGameplayTag RequiredState = FGameplayTag()) const;
+
+	/** Gets all implementing objects for an actor that are at RequiredState or later, other than excluding feature if specified */
+	void GetAllFeatureImplementers(TArray<UObject*>& OutImplementers, AActor* Actor, FGameplayTag RequiredState, FName ExcludingFeature = NAME_None) const;
+
+	/** Checks to see if all features of object, other than the excluding feature if specified, have reached a specified state or later */
+	bool HaveAllFeaturesReachedInitState(AActor* Actor, FGameplayTag RequiredState, FName ExcludingFeature = NAME_None) const;
+
+	/** Changes the current actor feature state, this will call registered callbacks and return true if anything changed */
+	bool ChangeFeatureInitState(AActor* Actor, FName FeatureName, UObject* Implementer, FGameplayTag FeatureState);
+
+	/** Removes an actor and all of it's state information */
+	void RemoveActorFeatureData(AActor* Actor);
+
+	/** Removes an implementing object and any feature states it implements */
+	void RemoveFeatureImplementer(AActor* Actor, UObject* Implementer);
+	
+	/** 
+	 * Registers native delegate for feature state change notifications on a specific actor and may call it immediately
+	 * 
+	 * @param Actor				The actor to listen for state changes to, if you don't have a specific actor call the Class version instead
+	 * @param FeatureName		If not empty, only listen to state changes for the specified feature
+	 * @param RequiredState		If specified, only activate if the init state of the feature is equal to or later than this
+	 * @param Delegate			Native delegate to call
+	 * @param bCallImmediately	If true and the actor feature is already in the specified state, call delegate immediately after registering 
+	 * @return DelegateHandle used for later removal
+	 */
+	FDelegateHandle RegisterAndCallForActorInitState(AActor* Actor, FName FeatureName, FGameplayTag RequiredState, FActorInitStateChangedDelegate Delegate, bool bCallImmediately = true);
+
+	/** Removes a registered delegate bound to a specific actor */
+	bool UnregisterActorInitStateDelegate(AActor* Actor, FDelegateHandle& Handle);
+	
+	/**
+	 * Registers blueprint delegate for feature state change notifications on a specific actor and may call it immediately
+	 *
+	 * @param Actor				The actor to listen for state changes to, if you don't have a specific actor call the Class version instead
+	 * @param FeatureName		If not empty, only listen to state changes for the specified feature
+	 * @param RequiredState		If specified, only activate if the init state of the feature is equal to or later than this
+	 * @param Delegate			Native delegate to call
+	 * @param bCallImmediately	If true and the actor feature is already in the specified state, call delegate immediately after registering
+	 * @return true if delegate was registered
+	 */
+	UFUNCTION(BlueprintCallable, Category = "InitState")
+	bool RegisterAndCallForActorInitState(AActor* Actor, FName FeatureName, FGameplayTag RequiredState, FActorInitStateChangedBPDelegate Delegate, bool bCallImmediately = true);
+
+	/** Removes a registered delegate bound to a specific actor */
+	UFUNCTION(BlueprintCallable, Category = "InitState")
+	bool UnregisterActorInitStateDelegate(AActor* Actor, FActorInitStateChangedBPDelegate DelegateToRemove);
+
+
+	/**
+	 * Registers native delegate for feature state change notifications on a class of actors and may call it immediately
+	 *
+	 * @param ActorClass		Name of an actor class to listen for changes to
+	 * @param FeatureName		If not empty, only listen to state changes for the specified feature
+	 * @param RequiredState		If specified, only activate if the init state of the feature is equal to or later than this
+	 * @param Delegate			Native delegate to call
+	 * @param bCallImmediately	If true and the actor feature is already in the specified state, call delegate immediately after registering
+	 * @return DelegateHandle used for later removal
+	 */
+	FDelegateHandle RegisterAndCallForClassInitState(const TSoftClassPtr<AActor>& ActorClass, FName FeatureName, FGameplayTag RequiredState, FActorInitStateChangedDelegate Delegate, bool bCallImmediately = true);
+
+	/** Removes a registered delegate bound to a class */
+	bool UnregisterClassInitStateDelegate(const TSoftClassPtr<AActor>& ActorClass, FDelegateHandle& Handle);
+
+	/**
+	 * Registers blueprint delegate for feature state change notifications on a class of actors and may call it immediately
+	 *
+	 * @param ActorClass		Name of an actor class to listen for changes to
+	 * @param FeatureName		If not empty, only listen to state changes for the specified feature
+	 * @param RequiredState		If specified, only activate if the init state of the feature is equal to or later than this
+	 * @param Delegate			Native delegate to call
+	 * @param bCallImmediately	If true and the actor feature is already in the specified state, call delegate immediately after registering
+	 * @return true if delegate was registered
+	 */
+	UFUNCTION(BlueprintCallable, Category = "InitState")
+	bool RegisterAndCallForClassInitState(TSoftClassPtr<AActor> ActorClass, FName FeatureName, FGameplayTag RequiredState, FActorInitStateChangedBPDelegate Delegate, bool bCallImmediately = true);
+
+	/** Removes a registered delegate bound to a class */
+	UFUNCTION(BlueprintCallable, Category = "InitState")
+	bool UnregisterClassInitStateDelegate(TSoftClassPtr<AActor> ActorClass, FActorInitStateChangedBPDelegate DelegateToRemove);
+
+
+private:
+	/** List of all registered feature states in order */
+	TArray<FGameplayTag> InitStateOrder;
+
+	/** State for a specific object implementing an actor feature, should this be in a map instead of an array? */
+	struct FActorFeatureState
+	{
+		FActorFeatureState(FName InFeatureName) : FeatureName(InFeatureName) {}
+
+		/** The feature this is tracking */
+		FName FeatureName;
+
+		/** The state when it was last registered */
+		FGameplayTag CurrentState;
+
+		/** The object implementing this feature, this can be null */
+		TWeakObjectPtr<UObject> Implementer;
+	};
+
+	/** Holds the list of feature delegates */
+	struct FActorFeatureRegisteredDelegate
+	{
+		/** Construct from a native or BP Delegate */
+		FActorFeatureRegisteredDelegate(FActorInitStateChangedDelegate&& InDelegate, FName InFeatureName = NAME_None, FGameplayTag InInitState = FGameplayTag());
+		FActorFeatureRegisteredDelegate(FActorInitStateChangedBPDelegate&& InDelegate, FName InFeatureName = NAME_None, FGameplayTag InInitState = FGameplayTag());
+
+		/** Call the appropriate native/bp delegate, this could invalidate this struct */
+		void Execute(AActor* OwningActor, FName FeatureName, UObject* Implementer, FGameplayTag FeatureState);
+
+		/** Delegate that is called on notification */
+		FActorInitStateChangedDelegate Delegate;
+
+		/** BP delegate that is called on notification */
+		FActorInitStateChangedBPDelegate BPDelegate;
+
+		/** A handle assigned to this delegate so it acts like a multicast delegate for removal */
+		FDelegateHandle DelegateHandle;
+
+		/** If this is not null, will only activate for specific feature names */
+		FName RequiredFeatureName;
+
+		/** If this is not null, will only activate for states >= to this */
+		FGameplayTag RequiredInitState;
+	};
+
+	/** Information for each registered actor */
+	struct FActorFeatureData
+	{
+		/** Actor class for cross referencing with the class callbacks */
+		TWeakObjectPtr<UClass> ActorClass;
+
+		/** All active features */
+		TArray<FActorFeatureState> RegisteredStates;
+
+		/** All delegates bound to this actor */
+		TArray<FActorFeatureRegisteredDelegate> RegisteredDelegates;
+	};
+
+	/** Actors that were registered as tracking feature state */
+	TMap<FObjectKey, FActorFeatureData> ActorFeatureMap;
+
+	/** Global delegates for any class name */
+	TMap<FComponentRequestReceiverClassPath, TArray<FActorFeatureRegisteredDelegate>> ClassFeatureChangeDelegates;
+
+	/** A queue of state changes to call delegates for, we don't want recursive callbacks */
+	TArray<TPair<AActor*, FActorFeatureState> > StateChangeQueue;
+
+	/** Position in state change queue, INDEX_NONE means not actively handling */
+	int32 CurrentStateChange;
+
+	/** Find an appropriate state struct if it exists */
+	const FActorFeatureState* FindFeatureStateStruct(const FActorFeatureData* ActorStruct, FName FeatureName, FGameplayTag RequiredState = FGameplayTag()) const;
+
+	/** Add to queue for delegate processing */
+	void ProcessFeatureStateChange(AActor* Actor, const FActorFeatureState* StateChange);
+
+	/** Call all delegates for a specific actor feature change */
+	void CallFeatureStateDelegates(AActor* Actor, FActorFeatureState StateChange);
+
+	/** Call the specified delegate for all matching features on the actor, this should be passed a copy of the original delegate */
+	void CallDelegateForMatchingFeatures(AActor* Actor, FActorFeatureRegisteredDelegate& RegisteredDelegate);
+
+	/** Call the specified delegate for all matching actors and features, this should be passed a copy of the original delegate */
+	void CallDelegateForMatchingActors(UClass* ActorClass, FActorFeatureRegisteredDelegate& RegisteredDelegate);
+
+	/** Gets or creates the actor struct */
+	FActorFeatureData& FindOrAddActorData(AActor* Actor);
+
+	/** Searches an array of delegates for the one with the specified handle, returns -1 if not found */
+	int32 GetIndexForRegisteredDelegate(TArray<FActorFeatureRegisteredDelegate>& DelegatesToSearch, FDelegateHandle SearchHandle) const;
+	int32 GetIndexForRegisteredDelegate(TArray<FActorFeatureRegisteredDelegate>& DelegatesToSearch, FActorInitStateChangedBPDelegate SearchDelegate) const;
+};
+
+
+// TODO Move into different header and see if we can make this pure virtual
+/** Interface that can be implemented by actors/components to make interacting with the feature state system easier */
+UINTERFACE(MinimalAPI, NotBlueprintable)
+class UGameFrameworkInitStateInterface : public UInterface
+{
+	GENERATED_BODY()
+};
+
+class MODULARGAMEPLAY_API IGameFrameworkInitStateInterface
+{
+	GENERATED_BODY()
+public:
+
+	/** Returns the Actor this object is bound to, might be this object */
+	virtual AActor* GetOwningActor() const;
+
+	/** Gets the component manager corresponding to this object based on owning actor */
+	UGameFrameworkComponentManager* GetComponentManager() const;
+
+	/** Returns the feature this object implements, this interface is only meant for simple objects with a single feature like Actor */
+	virtual FName GetFeatureName() const { return NAME_None; }
+
+	/** Returns the current feature state of this object, the default behavior is to query the manager */
+	virtual FGameplayTag GetInitState() const;
+
+	/** Checks the component manager to see if we have already reached the desired state or a later one */
+	virtual bool HasReachedInitState(FGameplayTag DesiredState) const;
+
+	/** Should be overridden to perform class-specific checks to see if the desired state can be reached */
+	virtual bool CanChangeInitState(UGameFrameworkComponentManager* Manager, FGameplayTag CurrentState, FGameplayTag DesiredState) const { return true; }
+
+	/** Should be overridden to perform class-specific state changes, this is called right before notifying the component manager */
+	virtual void HandleChangeInitState(UGameFrameworkComponentManager* Manager, FGameplayTag CurrentState, FGameplayTag DesiredState) {}
+
+	/** Checks to see if a change is possible, then calls execute and notify */
+	virtual bool TryToChangeInitState(FGameplayTag DesiredState);
+
+	/** Tries to follow a chain of connected init states, will progress states in order and returns the final state reached */
+	virtual FGameplayTag ContinueInitStateChain(const TArray<FGameplayTag>& InitStateChain);
+
+	/** Override to try and progress the default initialization path, likely using ContinueInitStateChain */
+	virtual void CheckDefaultInitialization() {}
+
+	/** This will call CheckDefaultInitialization on all other feature implementers using this interface, useful to update the state of any dependencies */
+	virtual void CheckDefaultInitializationForImplementers();
+
+	/** Signature for handling a game feature state, this is not registered by default */
+	virtual void OnActorInitStateChanged(AActor* OwningActor, FName FeatureName, UObject* Implementer, FGameplayTag FeatureState) {}
+
+	/** Call to bind the OnActorInitStateChanged function to the appropriate delegate on the component manager */
+	virtual void BindOnActorInitStateChanged(FName FeatureName, FGameplayTag RequiredState, bool bCallIfReached);
+
+	/** Unregisters state and delegate binding with component manager */
+	virtual void UnregisterInitStateData();
+
+	/** Binds a BP delegate to get called on a state change for this feature */
+	UFUNCTION(BlueprintCallable, Category = "InitState")
+	virtual bool RegisterAndCallForInitStateChange(FGameplayTag RequiredState, FActorInitStateChangedBPDelegate Delegate, bool bCallImmediately = true);
+
+	/** Unbinds a BP delegate from changes to this feature */
+	UFUNCTION(BlueprintCallable, Category = "InitState")
+	virtual bool UnregisterInitStateDelegate(FActorInitStateChangedBPDelegate Delegate);
+
+protected:
+	/** Default handle created from calling BindOnActorInitStateChanged */
+	FDelegateHandle ActorInitStateChangedHandle;
+
 };
