Don't store strong pointers for AllReceivers so failure to remove actors doesn't prevent world cleanup after PIE.

diff --git a/Engine/Plugins/Experimental/ModularGameplay/Source/ModularGameplay/Private/Components/GameFrameworkComponentManager.cpp b/Engine/Plugins/Experimental/ModularGameplay/Source/ModularGameplay/Private/Components/GameFrameworkComponentManager.cpp
--- a/Engine/Plugins/Experimental/ModularGameplay/Source/ModularGameplay/Private/Components/GameFrameworkComponentManager.cpp
+++ b/Engine/Plugins/Experimental/ModularGameplay/Source/ModularGameplay/Private/Components/GameFrameworkComponentManager.cpp
@@ -6,6 +6,7 @@
 #include "Engine/GameInstance.h"
 #include "EngineUtils.h"
 #include "ModularGameplayLogs.h"
+#include "UObject/UObjectGlobals.h"
 
 #if !UE_BUILD_SHIPPING
 #include "HAL/IConsoleManager.h"
@@ -71,7 +72,7 @@ void UGameFrameworkComponentManager::DumpGameFrameworkComponentManagers()
 			UE_LOG(LogModularGameplay, Display, TEXT("    Receivers... (Num:%d)"), Manager->AllReceivers.Num());
 			for (auto SetIt = Manager->AllReceivers.CreateConstIterator(); SetIt; ++SetIt)
 			{
-				UE_LOG(LogModularGameplay, Display, TEXT("      ReceiverInstance: %s"), *GetPathNameSafe(*SetIt));
+				UE_LOG(LogModularGameplay, Display, TEXT("      ReceiverInstance: %s"), *GetPathNameSafe(SetIt->ResolveObjectPtr()));
 			}
 #endif // WITH_EDITOR
 
@@ -103,6 +104,32 @@ void UGameFrameworkComponentManager::Initialize(FSubsystemCollectionBase& Collec
 	Super::Initialize(Collection);
 
 	CurrentStateChange = INDEX_NONE;
+
+#if WITH_EDITORONLY_DATA
+	FCoreUObjectDelegates::GetPostGarbageCollect().AddUObject(this, &UGameFrameworkComponentManager::PostGC);
+#endif
+}
+
+void UGameFrameworkComponentManager::Deinitialize() 
+{
+	Super::Deinitialize();
+#if WITH_EDITORONLY_DATA
+	FCoreUObjectDelegates::GetPostGarbageCollect().RemoveAll(this);
+#endif
+}
+
+void UGameFrameworkComponentManager::PostGC()
+{
+#if WITH_EDITORONLY_DATA
+	// Clear invalid receivers. 
+	for (auto It = AllReceivers.CreateIterator(); It; ++It)
+	{
+		if (It->ResolveObjectPtr() == nullptr)
+		{
+			It.RemoveCurrent();
+		}
+	}
+#endif
 }
 
 UGameFrameworkComponentManager* UGameFrameworkComponentManager::GetForActor(const AActor* Actor, bool bOnlyGameWorlds)
@@ -145,7 +172,7 @@ void UGameFrameworkComponentManager::AddReceiverInternal(AActor* Receiver)
 	checkSlow(Receiver);
 	
 #if WITH_EDITOR
-	AllReceivers.Add(Receiver);
+	AllReceivers.Add(FObjectKey(Receiver));
 #endif
 	
 	for (UClass* Class = Receiver->GetClass(); Class && Class != AActor::StaticClass(); Class = Class->GetSuperClass())
diff --git a/Engine/Plugins/Experimental/ModularGameplay/Source/ModularGameplay/Public/Components/GameFrameworkComponentManager.h b/Engine/Plugins/Experimental/ModularGameplay/Source/ModularGameplay/Public/Components/GameFrameworkComponentManager.h
--- a/Engine/Plugins/Experimental/ModularGameplay/Source/ModularGameplay/Public/Components/GameFrameworkComponentManager.h
+++ b/Engine/Plugins/Experimental/ModularGameplay/Source/ModularGameplay/Public/Components/GameFrameworkComponentManager.h
@@ -82,6 +82,7 @@ public:
 	using FExtensionHandlerDelegate = FExtensionHandlerDelegateInternal;
 
 	virtual void Initialize(FSubsystemCollectionBase& Collection) override;
+	virtual void Deinitialize() override;
 
 	/** Utility to get this manager from an actor, will return null if actor is null or not in a world */
 	static UGameFrameworkComponentManager* GetForActor(const AActor* Actor, bool bOnlyGameWorlds = true);
@@ -144,6 +145,10 @@ public:
 #endif // !UE_BUILD_SHIPPING
 
 private:
+#if WITH_EDITORONLY_DATA
+	void PostGC();
+#endif
+
 	void AddReceiverInternal(AActor* Receiver);
 	void RemoveReceiverInternal(AActor* Receiver);
 	void SendExtensionEventInternal(AActor* Receiver, const FName& EventName);
@@ -251,8 +256,7 @@ private:
 
 #if WITH_EDITORONLY_DATA
 	/** Editor-only set to validate that component requests are only being added for actors that call AddReceiver and RemoveReceiver */
-	UPROPERTY(transient)
-	TSet<AActor*> AllReceivers;
+	TSet<FObjectKey> AllReceivers;
 #endif // WITH_EDITORONLY_DATA
 
 	friend struct FComponentRequestHandle;
