Updates to the init state system:
Add a register feature function that can be called from component register as BeginPlay is too late to do the initial feature registration as another component may have fully initialized before BeginPlay is called
Change delegates to use a params struct
Split GameFrameworkInitStateInterface into separate header

Lyra-specific changes:
Fix bug with last version where if the hero component replicated slowly it wouldn't register until the extension component was totally done
Remove the old lyra-specific init system
Keep ability system and binding init/uninit separate from state system but clean up implementation and comments

diff --git a/Engine/Plugins/Experimental/ModularGameplay/Source/ModularGameplay/Private/Components/GameFrameworkComponentManager.cpp b/Engine/Plugins/Experimental/ModularGameplay/Source/ModularGameplay/Private/Components/GameFrameworkComponentManager.cpp
--- a/Engine/Plugins/Experimental/ModularGameplay/Source/ModularGameplay/Private/Components/GameFrameworkComponentManager.cpp
+++ b/Engine/Plugins/Experimental/ModularGameplay/Source/ModularGameplay/Private/Components/GameFrameworkComponentManager.cpp
@@ -530,8 +530,6 @@ void UGameFrameworkComponentManager::SendGameFrameworkComponentExtensionEvent(AA
 	}
 }
 
-// TODO Actor feature state prototypes
-
 UGameFrameworkComponentManager::FActorFeatureRegisteredDelegate::FActorFeatureRegisteredDelegate(FActorInitStateChangedDelegate&& InDelegate, FName InFeatureName, FGameplayTag InInitState)
 	: Delegate(InDelegate)
 	, DelegateHandle(FDelegateHandle::EGenerateNewHandleType::GenerateNewHandle)
@@ -552,15 +550,16 @@ UGameFrameworkComponentManager::FActorFeatureRegisteredDelegate::FActorFeatureRe
 
 void UGameFrameworkComponentManager::FActorFeatureRegisteredDelegate::Execute(AActor* OwningActor, FName FeatureName, UObject* Implementer, FGameplayTag FeatureState)
 {
+	FActorInitStateChangedParams Params(OwningActor, FeatureName, Implementer, FeatureState);
 	if (Delegate.IsBound())
 	{
 		ensure(!BPDelegate.IsBound());
 
-		Delegate.Execute(OwningActor, FeatureName, Implementer, FeatureState);
+		Delegate.Execute(Params);
 	}
 	else if (BPDelegate.IsBound())
 	{
-		BPDelegate.Execute(OwningActor, FeatureName, Implementer, FeatureState);
+		BPDelegate.Execute(Params);
 	}
 }
 
@@ -732,6 +731,35 @@ bool UGameFrameworkComponentManager::ChangeFeatureInitState(AActor* Actor, FName
 	return true;
 }
 
+bool UGameFrameworkComponentManager::RegisterFeatureImplementer(AActor* Actor, FName FeatureName, UObject* Implementer)
+{
+	if (Actor == nullptr || FeatureName.IsNone())
+	{
+		// TODO Ensure?
+		return false;
+	}
+
+	FActorFeatureData& ActorStruct = FindOrAddActorData(Actor);
+
+	FActorFeatureState* FoundState = nullptr;
+	for (FActorFeatureState& State : ActorStruct.RegisteredStates)
+	{
+		if (State.FeatureName == FeatureName)
+		{
+			// TODO what if it's already in the desired state?
+			FoundState = &State;
+		}
+	}
+
+	if (!FoundState)
+	{
+		FoundState = &ActorStruct.RegisteredStates.Emplace_GetRef(FeatureName);
+	}
+
+	FoundState->Implementer = Implementer;
+	return true;
+}
+
 void UGameFrameworkComponentManager::RemoveActorFeatureData(AActor* Actor)
 {
 	// TODO when should we clear up now-invalid actor keys?
@@ -1169,237 +1197,3 @@ int32 UGameFrameworkComponentManager::GetIndexForRegisteredDelegate(TArray<FActo
 
 	return INDEX_NONE;
 }
-
-// INTERFACE
-
-AActor* IGameFrameworkInitStateInterface::GetOwningActor() const
-{
-	// Removing const because most AActor functions can't handle const
-	AActor* FoundActor = const_cast<AActor*>(Cast<AActor>(this));
-
-	if (!FoundActor)
-	{
-		const UActorComponent* FoundComponent = Cast<UActorComponent>(this);
-		if (FoundComponent)
-		{
-			FoundActor = FoundComponent->GetOwner();
-		}
-	}
-
-	if (ensure(FoundActor))
-	{
-		// Subclasses must implement this if they are not actors or components
-		return FoundActor;
-	}
-
-	return nullptr;
-}
-
-UGameFrameworkComponentManager* IGameFrameworkInitStateInterface::GetComponentManager() const
-{
-	return UGameFrameworkComponentManager::GetForActor(GetOwningActor());
-}
-
-FGameplayTag IGameFrameworkInitStateInterface::GetInitState() const
-{
-	AActor* MyActor = GetOwningActor();
-	UGameFrameworkComponentManager* Manager = UGameFrameworkComponentManager::GetForActor(MyActor);
-	if (Manager)
-	{
-		return Manager->GetInitStateForFeature(MyActor, GetFeatureName());
-	}
-
-	return FGameplayTag();
-}
-
-bool IGameFrameworkInitStateInterface::HasReachedInitState(FGameplayTag DesiredState) const
-{
-	AActor* MyActor = GetOwningActor();
-	UGameFrameworkComponentManager* Manager = UGameFrameworkComponentManager::GetForActor(MyActor);
-	if (Manager)
-	{
-		return Manager->HasFeatureReachedInitState(MyActor, GetFeatureName(), DesiredState);
-	}
-
-	return false;
-}
-
-bool IGameFrameworkInitStateInterface::TryToChangeInitState(FGameplayTag DesiredState)
-{
-	UObject* ThisObject = Cast<UObject>(this);
-	AActor* MyActor = GetOwningActor();
-	UGameFrameworkComponentManager* Manager = UGameFrameworkComponentManager::GetForActor(MyActor);
-	const FName MyFeatureName = GetFeatureName();
-
-	if (!Manager || !ThisObject || !MyActor)
-	{
-		return false;	
-	}
-
-	FGameplayTag CurrentState = Manager->GetInitStateForFeature(MyActor, MyFeatureName);
-	
-	// If we're already in that state, just return
-	if (CurrentState == DesiredState)
-	{
-		return false;
-	}
-
-	if (!CanChangeInitState(Manager, CurrentState, DesiredState))
-	{
-		UE_LOG(LogModularGameplay, Verbose, TEXT("TryToChangeInitState: Cannot transition %s:%s (role %d) from %s to %s"),
-			*MyActor->GetName(), *MyFeatureName.ToString(), MyActor->GetLocalRole(), *CurrentState.ToString(), *DesiredState.ToString());
-		return false;
-	}
-
-	UE_LOG(LogModularGameplay, Verbose, TEXT("TryToChangeInitState: Transitioning %s:%s (role %d) from %s to %s"),
-		*MyActor->GetName(), *MyFeatureName.ToString(), MyActor->GetLocalRole(), *CurrentState.ToString(), *DesiredState.ToString());
-
-	// Perform the local change
-	HandleChangeInitState(Manager, CurrentState, DesiredState);
-
-	// The local change has completed, notify the system to register change and execute callbacks
-	return ensure(Manager->ChangeFeatureInitState(MyActor, MyFeatureName, ThisObject, DesiredState));
-}
-
-FGameplayTag IGameFrameworkInitStateInterface::ContinueInitStateChain(const TArray<FGameplayTag>& InitStateChain)
-{
-	UObject* ThisObject = Cast<UObject>(this);
-	AActor* MyActor = GetOwningActor();
-	UGameFrameworkComponentManager* Manager = UGameFrameworkComponentManager::GetForActor(MyActor);
-	const FName MyFeatureName = GetFeatureName();
-
-	if (!Manager || !ThisObject || !MyActor)
-	{
-		return FGameplayTag();
-	}
-
-	int32 ChainIndex = 0;
-	FGameplayTag CurrentState = Manager->GetInitStateForFeature(MyActor, MyFeatureName);
-
-	// For each state in chain before the last, see if we can transition to the next state
-	while (ChainIndex < InitStateChain.Num() - 1)
-	{
-		if (CurrentState == InitStateChain[ChainIndex])
-		{
-			FGameplayTag DesiredState = InitStateChain[ChainIndex + 1];
-			if (CanChangeInitState(Manager, CurrentState, DesiredState))
-			{
-				UE_LOG(LogModularGameplay, Verbose, TEXT("ContinueInitStateChain: Transitioning %s:%s (role %d) from %s to %s"),
-					*MyActor->GetName(), *MyFeatureName.ToString(), MyActor->GetLocalRole(), *CurrentState.ToString(), *DesiredState.ToString());
-
-				// Perform the local change
-				HandleChangeInitState(Manager, CurrentState, DesiredState);
-
-				// The local change has completed, notify the system to register change and execute callbacks
-				ensure(Manager->ChangeFeatureInitState(MyActor, MyFeatureName, ThisObject, DesiredState));
-
-				// Update state and check again
-				CurrentState = Manager->GetInitStateForFeature(MyActor, MyFeatureName);
-			}
-			else
-			{
-				UE_LOG(LogModularGameplay, Verbose, TEXT("ContinueInitStateChain: Cannot transition %s:%s (role %d) from %s to %s"),
-					*MyActor->GetName(), *MyFeatureName.ToString(), MyActor->GetLocalRole(), *CurrentState.ToString(), *DesiredState.ToString());
-			}
-		}
-
-		ChainIndex++;
-	}
-
-	return CurrentState;
-}
-
-void IGameFrameworkInitStateInterface::CheckDefaultInitializationForImplementers()
-{
-	UObject* ThisObject = Cast<UObject>(this);
-	AActor* MyActor = GetOwningActor();
-	UGameFrameworkComponentManager* Manager = UGameFrameworkComponentManager::GetForActor(MyActor);
-	const FName MyFeatureName = GetFeatureName();
-
-	if (Manager)
-	{
-		TArray<UObject*> Implementers;
-		Manager->GetAllFeatureImplementers(Implementers, MyActor, FGameplayTag(), MyFeatureName);
-
-		for (UObject* Implementer : Implementers)
-		{
-			if (IGameFrameworkInitStateInterface* ImplementerInterface = Cast<IGameFrameworkInitStateInterface>(Implementer))
-			{
-				ImplementerInterface->CheckDefaultInitialization();
-			}
-		}
-	}
-}
-
-void IGameFrameworkInitStateInterface::UnregisterInitStateData()
-{
-	UObject* ThisObject = Cast<UObject>(this);
-	AActor* MyActor = GetOwningActor();
-	UGameFrameworkComponentManager* Manager = UGameFrameworkComponentManager::GetForActor(MyActor);
-
-	if (Manager)
-	{
-		if (ThisObject == MyActor)
-		{
-			// This will clear all the feature states and delegates
-			Manager->RemoveActorFeatureData(MyActor);
-		}
-		else
-		{
-			Manager->RemoveFeatureImplementer(MyActor, ThisObject);
-		}
-
-		if (ActorInitStateChangedHandle.IsValid())
-		{
-			Manager->UnregisterActorInitStateDelegate(MyActor, ActorInitStateChangedHandle);
-		}
-	}
-}
-
-void IGameFrameworkInitStateInterface::BindOnActorInitStateChanged(FName FeatureName, FGameplayTag RequiredState, bool bCallIfReached)
-{
-	UObject* ThisObject = Cast<UObject>(this);
-	AActor* MyActor = GetOwningActor();
-	UGameFrameworkComponentManager* Manager = UGameFrameworkComponentManager::GetForActor(MyActor);
-
-	if (ensure(MyActor && Manager))
-	{
-		// Bind as a weak lambda because this is not a UObject but is guaranteed to be valid as long as ThisObject is
-		FActorInitStateChangedDelegate Delegate = FActorInitStateChangedDelegate::CreateWeakLambda(ThisObject,
-			[this](AActor* OwningActor, FName FeatureName, UObject* Implementer, FGameplayTag FeatureState)
-			{
-				this->OnActorInitStateChanged(OwningActor, FeatureName, Implementer, FeatureState);
-			});
-
-		ActorInitStateChangedHandle = Manager->RegisterAndCallForActorInitState(MyActor, FeatureName, RequiredState, MoveTemp(Delegate), bCallIfReached);
-	}
-}
-
-bool IGameFrameworkInitStateInterface::RegisterAndCallForInitStateChange(FGameplayTag RequiredState, FActorInitStateChangedBPDelegate Delegate, bool bCallImmediately)
-{
-	UObject* ThisObject = Cast<UObject>(this);
-	AActor* MyActor = GetOwningActor();
-	UGameFrameworkComponentManager* Manager = UGameFrameworkComponentManager::GetForActor(MyActor);
-	const FName MyFeatureName = GetFeatureName();
-
-	if (ensure(MyActor && Manager))
-	{
-		return Manager->RegisterAndCallForActorInitState(MyActor, MyFeatureName, RequiredState, Delegate, bCallImmediately);
-	}
-
-	return false;
-}
-
-bool IGameFrameworkInitStateInterface::UnregisterInitStateDelegate(FActorInitStateChangedBPDelegate Delegate)
-{
-	UObject* ThisObject = Cast<UObject>(this);
-	AActor* MyActor = GetOwningActor();
-	UGameFrameworkComponentManager* Manager = UGameFrameworkComponentManager::GetForActor(MyActor);
-
-	if (ensure(MyActor && Manager))
-	{
-		return Manager->UnregisterActorInitStateDelegate(MyActor, Delegate);
-	}
-
-	return false;
-}
diff --git a/Engine/Plugins/Experimental/ModularGameplay/Source/ModularGameplay/Private/Components/GameFrameworkInitStateInterface.cpp b/Engine/Plugins/Experimental/ModularGameplay/Source/ModularGameplay/Private/Components/GameFrameworkInitStateInterface.cpp
new file mode 100644
--- /dev/null
+++ b/Engine/Plugins/Experimental/ModularGameplay/Source/ModularGameplay/Private/Components/GameFrameworkInitStateInterface.cpp
@@ -0,0 +1,256 @@
+// Copyright Epic Games, Inc. All Rights Reserved.
+
+#include "Components/GameFrameworkInitStateInterface.h"
+#include "Components/GameFrameworkComponent.h"
+#include "Engine/World.h"
+#include "Engine/GameInstance.h"
+#include "EngineUtils.h"
+#include "ModularGameplayLogs.h"
+
+
+AActor* IGameFrameworkInitStateInterface::GetOwningActor() const
+{
+	// Removing const because most AActor functions can't handle const
+	AActor* FoundActor = const_cast<AActor*>(Cast<AActor>(this));
+
+	if (!FoundActor)
+	{
+		const UActorComponent* FoundComponent = Cast<UActorComponent>(this);
+		if (FoundComponent)
+		{
+			FoundActor = FoundComponent->GetOwner();
+		}
+	}
+
+	if (ensure(FoundActor))
+	{
+		// Subclasses must implement this if they are not actors or components
+		return FoundActor;
+	}
+
+	return nullptr;
+}
+
+UGameFrameworkComponentManager* IGameFrameworkInitStateInterface::GetComponentManager() const
+{
+	return UGameFrameworkComponentManager::GetForActor(GetOwningActor());
+}
+
+FGameplayTag IGameFrameworkInitStateInterface::GetInitState() const
+{
+	AActor* MyActor = GetOwningActor();
+	UGameFrameworkComponentManager* Manager = UGameFrameworkComponentManager::GetForActor(MyActor);
+	if (Manager)
+	{
+		return Manager->GetInitStateForFeature(MyActor, GetFeatureName());
+	}
+
+	return FGameplayTag();
+}
+
+bool IGameFrameworkInitStateInterface::HasReachedInitState(FGameplayTag DesiredState) const
+{
+	AActor* MyActor = GetOwningActor();
+	UGameFrameworkComponentManager* Manager = UGameFrameworkComponentManager::GetForActor(MyActor);
+	if (Manager)
+	{
+		return Manager->HasFeatureReachedInitState(MyActor, GetFeatureName(), DesiredState);
+	}
+
+	return false;
+}
+
+bool IGameFrameworkInitStateInterface::TryToChangeInitState(FGameplayTag DesiredState)
+{
+	UObject* ThisObject = Cast<UObject>(this);
+	AActor* MyActor = GetOwningActor();
+	UGameFrameworkComponentManager* Manager = UGameFrameworkComponentManager::GetForActor(MyActor);
+	const FName MyFeatureName = GetFeatureName();
+
+	if (!Manager || !ThisObject || !MyActor)
+	{
+		return false;	
+	}
+
+	FGameplayTag CurrentState = Manager->GetInitStateForFeature(MyActor, MyFeatureName);
+	
+	// If we're already in that state, just return
+	if (CurrentState == DesiredState)
+	{
+		return false;
+	}
+
+	if (!CanChangeInitState(Manager, CurrentState, DesiredState))
+	{
+		UE_LOG(LogModularGameplay, Verbose, TEXT("TryToChangeInitState: Cannot transition %s:%s (role %d) from %s to %s"),
+			*MyActor->GetName(), *MyFeatureName.ToString(), MyActor->GetLocalRole(), *CurrentState.ToString(), *DesiredState.ToString());
+		return false;
+	}
+
+	UE_LOG(LogModularGameplay, Verbose, TEXT("TryToChangeInitState: Transitioning %s:%s (role %d) from %s to %s"),
+		*MyActor->GetName(), *MyFeatureName.ToString(), MyActor->GetLocalRole(), *CurrentState.ToString(), *DesiredState.ToString());
+
+	// Perform the local change
+	HandleChangeInitState(Manager, CurrentState, DesiredState);
+
+	// The local change has completed, notify the system to register change and execute callbacks
+	return ensure(Manager->ChangeFeatureInitState(MyActor, MyFeatureName, ThisObject, DesiredState));
+}
+
+FGameplayTag IGameFrameworkInitStateInterface::ContinueInitStateChain(const TArray<FGameplayTag>& InitStateChain)
+{
+	UObject* ThisObject = Cast<UObject>(this);
+	AActor* MyActor = GetOwningActor();
+	UGameFrameworkComponentManager* Manager = UGameFrameworkComponentManager::GetForActor(MyActor);
+	const FName MyFeatureName = GetFeatureName();
+
+	if (!Manager || !ThisObject || !MyActor)
+	{
+		return FGameplayTag();
+	}
+
+	int32 ChainIndex = 0;
+	FGameplayTag CurrentState = Manager->GetInitStateForFeature(MyActor, MyFeatureName);
+
+	// For each state in chain before the last, see if we can transition to the next state
+	while (ChainIndex < InitStateChain.Num() - 1)
+	{
+		if (CurrentState == InitStateChain[ChainIndex])
+		{
+			FGameplayTag DesiredState = InitStateChain[ChainIndex + 1];
+			if (CanChangeInitState(Manager, CurrentState, DesiredState))
+			{
+				UE_LOG(LogModularGameplay, Verbose, TEXT("ContinueInitStateChain: Transitioning %s:%s (role %d) from %s to %s"),
+					*MyActor->GetName(), *MyFeatureName.ToString(), MyActor->GetLocalRole(), *CurrentState.ToString(), *DesiredState.ToString());
+
+				// Perform the local change
+				HandleChangeInitState(Manager, CurrentState, DesiredState);
+
+				// The local change has completed, notify the system to register change and execute callbacks
+				ensure(Manager->ChangeFeatureInitState(MyActor, MyFeatureName, ThisObject, DesiredState));
+
+				// Update state and check again
+				CurrentState = Manager->GetInitStateForFeature(MyActor, MyFeatureName);
+			}
+			else
+			{
+				UE_LOG(LogModularGameplay, Verbose, TEXT("ContinueInitStateChain: Cannot transition %s:%s (role %d) from %s to %s"),
+					*MyActor->GetName(), *MyFeatureName.ToString(), MyActor->GetLocalRole(), *CurrentState.ToString(), *DesiredState.ToString());
+			}
+		}
+
+		ChainIndex++;
+	}
+
+	return CurrentState;
+}
+
+void IGameFrameworkInitStateInterface::CheckDefaultInitializationForImplementers()
+{
+	UObject* ThisObject = Cast<UObject>(this);
+	AActor* MyActor = GetOwningActor();
+	UGameFrameworkComponentManager* Manager = UGameFrameworkComponentManager::GetForActor(MyActor);
+	const FName MyFeatureName = GetFeatureName();
+
+	if (Manager)
+	{
+		TArray<UObject*> Implementers;
+		Manager->GetAllFeatureImplementers(Implementers, MyActor, FGameplayTag(), MyFeatureName);
+
+		for (UObject* Implementer : Implementers)
+		{
+			if (IGameFrameworkInitStateInterface* ImplementerInterface = Cast<IGameFrameworkInitStateInterface>(Implementer))
+			{
+				ImplementerInterface->CheckDefaultInitialization();
+			}
+		}
+	}
+}
+
+
+void IGameFrameworkInitStateInterface::RegisterInitStateFeature()
+{
+	UObject* ThisObject = Cast<UObject>(this);
+	AActor* MyActor = GetOwningActor();
+	UGameFrameworkComponentManager* Manager = UGameFrameworkComponentManager::GetForActor(MyActor);
+	const FName MyFeatureName = GetFeatureName();
+
+	if (MyActor && Manager)
+	{
+		// Manager will be null if this isn't in a game world
+		Manager->RegisterFeatureImplementer(MyActor, MyFeatureName, ThisObject);
+	}
+}
+
+void IGameFrameworkInitStateInterface::UnregisterInitStateFeature()
+{
+	UObject* ThisObject = Cast<UObject>(this);
+	AActor* MyActor = GetOwningActor();
+	UGameFrameworkComponentManager* Manager = UGameFrameworkComponentManager::GetForActor(MyActor);
+
+	if (Manager)
+	{
+		if (ThisObject == MyActor)
+		{
+			// This will clear all the feature states and delegates
+			Manager->RemoveActorFeatureData(MyActor);
+		}
+		else
+		{
+			Manager->RemoveFeatureImplementer(MyActor, ThisObject);
+		}
+
+		if (ActorInitStateChangedHandle.IsValid())
+		{
+			Manager->UnregisterActorInitStateDelegate(MyActor, ActorInitStateChangedHandle);
+		}
+	}
+}
+
+void IGameFrameworkInitStateInterface::BindOnActorInitStateChanged(FName FeatureName, FGameplayTag RequiredState, bool bCallIfReached)
+{
+	UObject* ThisObject = Cast<UObject>(this);
+	AActor* MyActor = GetOwningActor();
+	UGameFrameworkComponentManager* Manager = UGameFrameworkComponentManager::GetForActor(MyActor);
+
+	if (ensure(MyActor && Manager))
+	{
+		// Bind as a weak lambda because this is not a UObject but is guaranteed to be valid as long as ThisObject is
+		FActorInitStateChangedDelegate Delegate = FActorInitStateChangedDelegate::CreateWeakLambda(ThisObject,
+			[this](const FActorInitStateChangedParams& Params)
+			{
+				this->OnActorInitStateChanged(Params);
+			});
+
+		ActorInitStateChangedHandle = Manager->RegisterAndCallForActorInitState(MyActor, FeatureName, RequiredState, MoveTemp(Delegate), bCallIfReached);
+	}
+}
+
+bool IGameFrameworkInitStateInterface::RegisterAndCallForInitStateChange(FGameplayTag RequiredState, FActorInitStateChangedBPDelegate Delegate, bool bCallImmediately)
+{
+	UObject* ThisObject = Cast<UObject>(this);
+	AActor* MyActor = GetOwningActor();
+	UGameFrameworkComponentManager* Manager = UGameFrameworkComponentManager::GetForActor(MyActor);
+	const FName MyFeatureName = GetFeatureName();
+
+	if (ensure(MyActor && Manager))
+	{
+		return Manager->RegisterAndCallForActorInitState(MyActor, MyFeatureName, RequiredState, Delegate, bCallImmediately);
+	}
+
+	return false;
+}
+
+bool IGameFrameworkInitStateInterface::UnregisterInitStateDelegate(FActorInitStateChangedBPDelegate Delegate)
+{
+	UObject* ThisObject = Cast<UObject>(this);
+	AActor* MyActor = GetOwningActor();
+	UGameFrameworkComponentManager* Manager = UGameFrameworkComponentManager::GetForActor(MyActor);
+
+	if (ensure(MyActor && Manager))
+	{
+		return Manager->UnregisterActorInitStateDelegate(MyActor, Delegate);
+	}
+
+	return false;
+}
diff --git a/Engine/Plugins/Experimental/ModularGameplay/Source/ModularGameplay/Public/Components/GameFrameworkComponentManager.h b/Engine/Plugins/Experimental/ModularGameplay/Source/ModularGameplay/Public/Components/GameFrameworkComponentManager.h
--- a/Engine/Plugins/Experimental/ModularGameplay/Source/ModularGameplay/Public/Components/GameFrameworkComponentManager.h
+++ b/Engine/Plugins/Experimental/ModularGameplay/Source/ModularGameplay/Public/Components/GameFrameworkComponentManager.h
@@ -16,11 +16,6 @@ class AActor;
 class UActorComponent;
 class UGameFrameworkComponent;
 
-/** Native delegate called when an actor feature changes init state */
-DECLARE_DELEGATE_FourParams(FActorInitStateChangedDelegate, AActor* /* OwningActor */, FName /* FeatureName */, UObject* /* Implementer */, FGameplayTag /* FeatureState */);
-
-/** Blueprint delegate called when an actor feature changes init state */
-DECLARE_DYNAMIC_DELEGATE_FourParams(FActorInitStateChangedBPDelegate, AActor*, OwningActor, FName, FeatureName, UObject*, Implementer, FGameplayTag, FeatureState);
 
 /** 
  * A handle for a request to put components or call a delegate for an extensible actor class.
@@ -59,6 +54,49 @@ private:
 	FDelegateHandle ExtensionHandle;
 };
 
+
+/** Parameters struct for Init State change functions */
+USTRUCT(BlueprintType)
+struct FActorInitStateChangedParams
+{
+	GENERATED_BODY()
+
+	FActorInitStateChangedParams()
+		: OwningActor(nullptr)
+		, Implementer(nullptr)
+	{}
+
+	FActorInitStateChangedParams(AActor* InOwningActor, FName InFeatureName, UObject* InImplementer, FGameplayTag InFeatureState)
+		: OwningActor(InOwningActor)
+		, FeatureName(InFeatureName)
+		, Implementer(InImplementer)
+		, FeatureState(InFeatureState)
+	{}
+
+	/** The actor owning the feature that changed */
+	UPROPERTY(BlueprintReadWrite, EditAnywhere, Category = Default)
+	AActor* OwningActor;
+
+	/** Name of the feature that changed */
+	UPROPERTY(BlueprintReadWrite, EditAnywhere, Category = Default)
+	FName FeatureName;
+
+	/** The object (often a component) that implements the feature */
+	UPROPERTY(BlueprintReadWrite, EditAnywhere, Category = Default)
+	UObject* Implementer;
+
+	/** The new state of the feature */
+	UPROPERTY(BlueprintReadWrite, EditAnywhere, Category = Default)
+	FGameplayTag FeatureState;
+};
+
+/** Native delegate called when an actor feature changes init state */
+DECLARE_DELEGATE_OneParam(FActorInitStateChangedDelegate, const FActorInitStateChangedParams&);
+
+/** Blueprint delegate called when an actor feature changes init state */
+DECLARE_DYNAMIC_DELEGATE_OneParam(FActorInitStateChangedBPDelegate, const FActorInitStateChangedParams&, Params);
+
+
 /** 
  * GameFrameworkComponentManager
  *
@@ -261,8 +299,10 @@ private:
 
 	friend struct FComponentRequestHandle;
 
-	// TODO work in progress for new actor feature ini state system
 
+	//////////////////////////////////////////////////////////////////////////////////////////////
+	// The init state system can be used by components to coordinate their initialization using game-specific states specified as gameplay tags
+	// IGameFrameworkInitStateInterface provides a simple implementation that can be inherted by components
 
 public:
 
@@ -290,6 +330,9 @@ public:
 	/** Changes the current actor feature state, this will call registered callbacks and return true if anything changed */
 	bool ChangeFeatureInitState(AActor* Actor, FName FeatureName, UObject* Implementer, FGameplayTag FeatureState);
 
+	/** Registers an implementer for a given feature, this will create a feature if required and set the implementer object but will not change the current state */
+	bool RegisterFeatureImplementer(AActor* Actor, FName FeatureName, UObject* Implementer);
+
 	/** Removes an actor and all of it's state information */
 	void RemoveActorFeatureData(AActor* Actor);
 
@@ -454,73 +497,3 @@ private:
 	int32 GetIndexForRegisteredDelegate(TArray<FActorFeatureRegisteredDelegate>& DelegatesToSearch, FDelegateHandle SearchHandle) const;
 	int32 GetIndexForRegisteredDelegate(TArray<FActorFeatureRegisteredDelegate>& DelegatesToSearch, FActorInitStateChangedBPDelegate SearchDelegate) const;
 };
-
-
-// TODO Move into different header and see if we can make this pure virtual
-/** Interface that can be implemented by actors/components to make interacting with the feature state system easier */
-UINTERFACE(MinimalAPI, NotBlueprintable)
-class UGameFrameworkInitStateInterface : public UInterface
-{
-	GENERATED_BODY()
-};
-
-class MODULARGAMEPLAY_API IGameFrameworkInitStateInterface
-{
-	GENERATED_BODY()
-public:
-
-	/** Returns the Actor this object is bound to, might be this object */
-	virtual AActor* GetOwningActor() const;
-
-	/** Gets the component manager corresponding to this object based on owning actor */
-	UGameFrameworkComponentManager* GetComponentManager() const;
-
-	/** Returns the feature this object implements, this interface is only meant for simple objects with a single feature like Actor */
-	virtual FName GetFeatureName() const { return NAME_None; }
-
-	/** Returns the current feature state of this object, the default behavior is to query the manager */
-	virtual FGameplayTag GetInitState() const;
-
-	/** Checks the component manager to see if we have already reached the desired state or a later one */
-	virtual bool HasReachedInitState(FGameplayTag DesiredState) const;
-
-	/** Should be overridden to perform class-specific checks to see if the desired state can be reached */
-	virtual bool CanChangeInitState(UGameFrameworkComponentManager* Manager, FGameplayTag CurrentState, FGameplayTag DesiredState) const { return true; }
-
-	/** Should be overridden to perform class-specific state changes, this is called right before notifying the component manager */
-	virtual void HandleChangeInitState(UGameFrameworkComponentManager* Manager, FGameplayTag CurrentState, FGameplayTag DesiredState) {}
-
-	/** Checks to see if a change is possible, then calls execute and notify */
-	virtual bool TryToChangeInitState(FGameplayTag DesiredState);
-
-	/** Tries to follow a chain of connected init states, will progress states in order and returns the final state reached */
-	virtual FGameplayTag ContinueInitStateChain(const TArray<FGameplayTag>& InitStateChain);
-
-	/** Override to try and progress the default initialization path, likely using ContinueInitStateChain */
-	virtual void CheckDefaultInitialization() {}
-
-	/** This will call CheckDefaultInitialization on all other feature implementers using this interface, useful to update the state of any dependencies */
-	virtual void CheckDefaultInitializationForImplementers();
-
-	/** Signature for handling a game feature state, this is not registered by default */
-	virtual void OnActorInitStateChanged(AActor* OwningActor, FName FeatureName, UObject* Implementer, FGameplayTag FeatureState) {}
-
-	/** Call to bind the OnActorInitStateChanged function to the appropriate delegate on the component manager */
-	virtual void BindOnActorInitStateChanged(FName FeatureName, FGameplayTag RequiredState, bool bCallIfReached);
-
-	/** Unregisters state and delegate binding with component manager */
-	virtual void UnregisterInitStateData();
-
-	/** Binds a BP delegate to get called on a state change for this feature */
-	UFUNCTION(BlueprintCallable, Category = "InitState")
-	virtual bool RegisterAndCallForInitStateChange(FGameplayTag RequiredState, FActorInitStateChangedBPDelegate Delegate, bool bCallImmediately = true);
-
-	/** Unbinds a BP delegate from changes to this feature */
-	UFUNCTION(BlueprintCallable, Category = "InitState")
-	virtual bool UnregisterInitStateDelegate(FActorInitStateChangedBPDelegate Delegate);
-
-protected:
-	/** Default handle created from calling BindOnActorInitStateChanged */
-	FDelegateHandle ActorInitStateChangedHandle;
-
-};
diff --git a/Engine/Plugins/Experimental/ModularGameplay/Source/ModularGameplay/Public/Components/GameFrameworkInitStateInterface.h b/Engine/Plugins/Experimental/ModularGameplay/Source/ModularGameplay/Public/Components/GameFrameworkInitStateInterface.h
new file mode 100644
--- /dev/null
+++ b/Engine/Plugins/Experimental/ModularGameplay/Source/ModularGameplay/Public/Components/GameFrameworkInitStateInterface.h
@@ -0,0 +1,79 @@
+// Copyright Epic Games, Inc. All Rights Reserved.
+
+#pragma once
+
+#include "GameFrameworkComponentManager.h"
+#include "GameFrameworkInitStateInterface.generated.h"
+
+/** Interface that can be implemented by actors/components to make interacting with the init state system easier */
+UINTERFACE(MinimalAPI, NotBlueprintable)
+class UGameFrameworkInitStateInterface : public UInterface
+{
+	GENERATED_BODY()
+};
+
+class MODULARGAMEPLAY_API IGameFrameworkInitStateInterface
+{
+	GENERATED_BODY()
+public:
+
+	/** Returns the Actor this object is bound to, might be this object */
+	virtual AActor* GetOwningActor() const;
+
+	/** Gets the component manager corresponding to this object based on owning actor */
+	UGameFrameworkComponentManager* GetComponentManager() const;
+
+	/** Returns the feature this object implements, this interface is only meant for simple objects with a single feature like Actor */
+	UFUNCTION(BlueprintCallable, Category = "InitState")
+	virtual FName GetFeatureName() const { return NAME_None; }
+
+	/** Returns the current feature state of this object, the default behavior is to query the manager */
+	UFUNCTION(BlueprintCallable, Category = "InitState")
+	virtual FGameplayTag GetInitState() const;
+
+	/** Checks the component manager to see if we have already reached the desired state or a later one */
+	UFUNCTION(BlueprintCallable, Category = "InitState")
+	virtual bool HasReachedInitState(FGameplayTag DesiredState) const;
+
+	/** Should be overridden to perform class-specific checks to see if the desired state can be reached */
+	virtual bool CanChangeInitState(UGameFrameworkComponentManager* Manager, FGameplayTag CurrentState, FGameplayTag DesiredState) const { return true; }
+
+	/** Should be overridden to perform class-specific state changes, this is called right before notifying the component manager */
+	virtual void HandleChangeInitState(UGameFrameworkComponentManager* Manager, FGameplayTag CurrentState, FGameplayTag DesiredState) {}
+
+	/** Checks to see if a change is possible, then calls execute and notify */
+	virtual bool TryToChangeInitState(FGameplayTag DesiredState);
+
+	/** Tries to follow a chain of connected init states, will progress states in order and returns the final state reached */
+	virtual FGameplayTag ContinueInitStateChain(const TArray<FGameplayTag>& InitStateChain);
+
+	/** Override to try and progress the default initialization path, likely using ContinueInitStateChain */
+	virtual void CheckDefaultInitialization() {}
+
+	/** This will call CheckDefaultInitialization on all other feature implementers using this interface, useful to update the state of any dependencies */
+	virtual void CheckDefaultInitializationForImplementers();
+
+	/** Signature for handling a game feature state, this is not registered by default */
+	virtual void OnActorInitStateChanged(const FActorInitStateChangedParams& Params) {}
+
+	/** Call to bind the OnActorInitStateChanged function to the appropriate delegate on the component manager */
+	virtual void BindOnActorInitStateChanged(FName FeatureName, FGameplayTag RequiredState, bool bCallIfReached);
+
+	/** Call to register with the component manager during spawn if this is a game world */
+	virtual void RegisterInitStateFeature();
+
+	/** Unregisters state and delegate binding with component manager */
+	virtual void UnregisterInitStateFeature();
+
+	/** Binds a BP delegate to get called on a state change for this feature */
+	UFUNCTION(BlueprintCallable, Category = "InitState")
+	virtual bool RegisterAndCallForInitStateChange(FGameplayTag RequiredState, FActorInitStateChangedBPDelegate Delegate, bool bCallImmediately = true);
+
+	/** Unbinds a BP delegate from changes to this feature */
+	UFUNCTION(BlueprintCallable, Category = "InitState")
+	virtual bool UnregisterInitStateDelegate(FActorInitStateChangedBPDelegate Delegate);
+
+protected:
+	/** Default handle created from calling BindOnActorInitStateChanged */
+	FDelegateHandle ActorInitStateChangedHandle;
+};
