[ModularGameplay]
* GameFrameworkComponentDelegates - Removed struct fwd declaration that shouldnt be there

diff --git a/Engine/Plugins/Experimental/ModularGameplay/Source/ModularGameplay/Public/Components/GameFrameworkComponentDelegates.h b/Engine/Plugins/Experimental/ModularGameplay/Source/ModularGameplay/Public/Components/GameFrameworkComponentDelegates.h
--- a/Engine/Plugins/Experimental/ModularGameplay/Source/ModularGameplay/Public/Components/GameFrameworkComponentDelegates.h
+++ b/Engine/Plugins/Experimental/ModularGameplay/Source/ModularGameplay/Public/Components/GameFrameworkComponentDelegates.h
@@ -2,8 +2,6 @@
 
 #pragma once
 
-struct FActorInitStateChangedParams;
-
 #include "GameplayTagContainer.h"
 #include "GameFrameworkComponentDelegates.generated.h"
 
