Added null check to fix crash in UGameFrameworkComponentManager::CallDelegateForMatchingActors.

diff --git a/Engine/Plugins/Experimental/ModularGameplay/Source/ModularGameplay/Private/Components/GameFrameworkComponentManager.cpp b/Engine/Plugins/Experimental/ModularGameplay/Source/ModularGameplay/Private/Components/GameFrameworkComponentManager.cpp
--- a/Engine/Plugins/Experimental/ModularGameplay/Source/ModularGameplay/Private/Components/GameFrameworkComponentManager.cpp
+++ b/Engine/Plugins/Experimental/ModularGameplay/Source/ModularGameplay/Private/Components/GameFrameworkComponentManager.cpp
@@ -1145,7 +1145,7 @@ void UGameFrameworkComponentManager::CallDelegateForMatchingActors(UClass* Actor
 	{
 		UClass* CheckActorClass = Pair.Value.ActorClass.Get();
 
-		if (CheckActorClass->IsChildOf(ActorClass))
+		if (CheckActorClass && CheckActorClass->IsChildOf(ActorClass))
 		{
 			AActor* FoundActor = Cast<AActor>(Pair.Key.ResolveObjectPtr());
 			if (FoundActor)
