GameStateComponent uses the UWorld::GetAuthGameMode to return the game mode since the GameState doesn't have a valid game mode during game state creation.
Added GameStateComponent::HandleMatchHasEnded hook

diff --git a/Engine/Plugins/Experimental/ModularGameplay/Source/ModularGameplay/Public/Components/GameStateComponent.h b/Engine/Plugins/Experimental/ModularGameplay/Source/ModularGameplay/Public/Components/GameStateComponent.h
--- a/Engine/Plugins/Experimental/ModularGameplay/Source/ModularGameplay/Public/Components/GameStateComponent.h
+++ b/Engine/Plugins/Experimental/ModularGameplay/Source/ModularGameplay/Public/Components/GameStateComponent.h
@@ -4,6 +4,7 @@
 
 #include "GameFramework/GameStateBase.h"
 #include "GameFrameworkComponent.h"
+#include "Engine/World.h"
 #include "GameStateComponent.generated.h"
 
 /**
@@ -41,8 +42,10 @@ public:
 	template <class T>
 	T* GetGameMode() const
 	{
+		// Note: Intentionally getting the game mode from the world instead of the game state as it can be null during game state initialization
 		static_assert(TPointerIsConvertibleFromTo<T, AGameModeBase>::Value, "'T' template parameter to GetGameMode must be derived from AGameModeBase");
-		return Cast<T>(GetGameStateChecked<AGameStateBase>()->AuthorityGameMode);
+		const UWorld* World = GetWorld();
+		return World ? World->GetAuthGameMode<T>() : nullptr;
 	}
 
 public:
@@ -53,6 +56,9 @@ public:
 
 	/** Called when gameplay has fully started */
 	virtual void HandleMatchHasStarted() {}
+
+	/** Called when gameplay has fully ended */
+	virtual void HandleMatchHasEnded() {}
 };
 
 #if UE_ENABLE_INCLUDE_ORDER_DEPRECATED_IN_5_2
