Use an inline allocator for FComponentRequestReceiverClassPath::Path. Improves perf during level streaming by avoiding temp allocations.

diff --git a/Engine/Plugins/Experimental/ModularGameplay/Source/ModularGameplay/Public/Components/GameFrameworkComponentManager.h b/Engine/Plugins/Experimental/ModularGameplay/Source/ModularGameplay/Public/Components/GameFrameworkComponentManager.h
--- a/Engine/Plugins/Experimental/ModularGameplay/Source/ModularGameplay/Public/Components/GameFrameworkComponentManager.h
+++ b/Engine/Plugins/Experimental/ModularGameplay/Source/ModularGameplay/Public/Components/GameFrameworkComponentManager.h
@@ -166,7 +166,8 @@ private:
 	/** A list of FNames to represent an object path. Used for fast hashing and comparison of paths */
 	struct FComponentRequestReceiverClassPath
 	{
-		TArray<FName> Path;
+		// Usually only 2 elements long. Ex [UPackage FName, UClass FName]
+		TArray<FName, TInlineAllocator<2>> Path;
 
 		FComponentRequestReceiverClassPath() {}
 
