Adding an AdditionFlags member to "Add Component" Game Feature Action so we can specify how component will be added:

-AddUnique - "If AddUnique is set, it will be added only if there is no existing component on ActorInstance of ComponentClass"
-AddUniqueIfNotChild - "If Add Unique If Not Child is set, it will be added only if ComponentClass is not a child (or same class) of an existing component on ActorInstance"
-UseAutoGeneratedName - "If UseAutoGeneratedName is set, it will generate a new name and not re-use the class name directly (which can lead to component recycling)."

diff --git a/Engine/Plugins/Runtime/GameFeatures/Source/GameFeatures/Private/GameFeatureAction_AddComponents.cpp b/Engine/Plugins/Runtime/GameFeatures/Source/GameFeatures/Private/GameFeatureAction_AddComponents.cpp
--- a/Engine/Plugins/Runtime/GameFeatures/Source/GameFeatures/Private/GameFeatureAction_AddComponents.cpp
+++ b/Engine/Plugins/Runtime/GameFeatures/Source/GameFeatures/Private/GameFeatureAction_AddComponents.cpp
@@ -15,6 +15,15 @@
 
 #define LOCTEXT_NAMESPACE "GameFeatures"
 
+//////////////////////////////////////////////////////////////////////
+// FGameFeatureComponentEntry
+FGameFeatureComponentEntry::FGameFeatureComponentEntry()
+	: bClientComponent(true)
+	, bServerComponent(true)
+	, AdditionFlags(static_cast<uint8>(EGameFrameworkAddComponentFlags::None))
+{
+}
+
 //////////////////////////////////////////////////////////////////////
 // UGameFeatureAction_AddComponents
 
@@ -120,7 +129,7 @@ void UGameFeatureAction_AddComponents::AddToWorld(const FWorldContext& WorldCont
 						TSubclassOf<UActorComponent> ComponentClass = Entry.ComponentClass.LoadSynchronous();
 						if (ComponentClass)
 						{
-							Handles.ComponentRequestHandles.Add(GFCM->AddComponentRequest(Entry.ActorClass, ComponentClass));
+							Handles.ComponentRequestHandles.Add(GFCM->AddComponentRequest(Entry.ActorClass, ComponentClass, static_cast<EGameFrameworkAddComponentFlags>(Entry.AdditionFlags)));
 						}
 						else if (!Entry.ComponentClass.IsNull())
 						{
diff --git a/Engine/Plugins/Runtime/GameFeatures/Source/GameFeatures/Public/GameFeatureAction_AddComponents.h b/Engine/Plugins/Runtime/GameFeatures/Source/GameFeatures/Public/GameFeatureAction_AddComponents.h
--- a/Engine/Plugins/Runtime/GameFeatures/Source/GameFeatures/Public/GameFeatureAction_AddComponents.h
+++ b/Engine/Plugins/Runtime/GameFeatures/Source/GameFeatures/Public/GameFeatureAction_AddComponents.h
@@ -14,6 +14,8 @@ class UGameInstance;
 struct FComponentRequestHandle;
 struct FWorldContext;
 
+enum class EGameFrameworkAddComponentFlags : uint8;
+
 // Description of a component to add to a type of actor when this game feature is enabled
 // (the actor class must be game feature aware, it does not happen magically)
 //@TODO: Write more documentation here about how to make an actor game feature / modular gameplay aware
@@ -22,6 +24,8 @@ struct GAMEFEATURES_API FGameFeatureComponentEntry
 {
 	GENERATED_BODY()
 
+	FGameFeatureComponentEntry();
+
 	// The base actor class to add a component to
 	UPROPERTY(EditAnywhere, Category="Components", meta=(AllowAbstract="True"))
 	TSoftClassPtr<AActor> ActorClass;
@@ -38,11 +42,9 @@ struct GAMEFEATURES_API FGameFeatureComponentEntry
 	UPROPERTY(EditAnywhere, Category="Components")
 	uint8 bServerComponent:1;
 
-	FGameFeatureComponentEntry()
-		: bClientComponent(true)
-		, bServerComponent(true)
-	{
-	}
+	// Observe these rules when adding the component, if any
+	UPROPERTY(EditAnywhere, Category = "Components", meta = (Bitmask, BitmaskEnum = "/Script/ModularGameplay.EGameFrameworkAddComponentFlags"))
+	uint8 AdditionFlags;
 };	
 
 //////////////////////////////////////////////////////////////////////
diff --git a/Engine/Plugins/Runtime/ModularGameplay/Source/ModularGameplay/Private/Components/GameFrameworkComponentManager.cpp b/Engine/Plugins/Runtime/ModularGameplay/Source/ModularGameplay/Private/Components/GameFrameworkComponentManager.cpp
--- a/Engine/Plugins/Runtime/ModularGameplay/Source/ModularGameplay/Private/Components/GameFrameworkComponentManager.cpp
+++ b/Engine/Plugins/Runtime/ModularGameplay/Source/ModularGameplay/Private/Components/GameFrameworkComponentManager.cpp
@@ -51,7 +51,7 @@ void UGameFrameworkComponentManager::AddReferencedObjects(UObject* InThis, FRefe
 		{
 			for (auto& ValueElement : MapIt.Value())
 			{
-				Collector.AddReferencedObject(ValueElement);
+				Collector.AddReferencedObject(ValueElement.Class);
 			}
 		}
 	}
@@ -88,9 +88,9 @@ void UGameFrameworkComponentManager::DumpGameFrameworkComponentManagers()
 			for (auto MapIt = Manager->ReceiverClassToComponentClassMap.CreateConstIterator(); MapIt; ++MapIt)
 			{
 				UE_LOG(LogModularGameplay, Display, TEXT("      RequestReceiverClass: %s (Num:%d)"), *MapIt.Key().ToDebugString(), MapIt.Value().Num());
-				for (UClass* ComponentClass : MapIt.Value())
+				for (const FComponentRequestInfo& ReceiverInfo : MapIt.Value())
 				{
-					UE_LOG(LogModularGameplay, Display, TEXT("        RequestComponentClass: %s"), *GetPathNameSafe(ComponentClass));
+					UE_LOG(LogModularGameplay, Display, TEXT("        RequestComponentClass: %s	AdditionFlags: %u"), *GetPathNameSafe(ReceiverInfo.Class), ReceiverInfo.AdditionFlags);
 				}
 			}
 		}
@@ -177,13 +177,13 @@ void UGameFrameworkComponentManager::AddReceiverInternal(AActor* Receiver)
 	for (UClass* Class = Receiver->GetClass(); Class && Class != AActor::StaticClass(); Class = Class->GetSuperClass())
 	{
 		FComponentRequestReceiverClassPath ReceiverClassPath(Class);
-		if (auto* ComponentClasses = ReceiverClassToComponentClassMap.Find(ReceiverClassPath))
+		if (auto* RequestInfoSet = ReceiverClassToComponentClassMap.Find(ReceiverClassPath))
 		{
-			for (UClass* ComponentClass : *ComponentClasses)
+			for (const FComponentRequestInfo& SetInfo : *RequestInfoSet)
 			{
-				if (ComponentClass)
+				if (SetInfo.Class)
 				{
-					CreateComponentOnInstance(Receiver, ComponentClass);
+					CreateComponentOnInstance(Receiver, SetInfo.Class, SetInfo.AdditionFlags);
 				}
 			}
 		}
@@ -239,7 +239,7 @@ void UGameFrameworkComponentManager::RemoveReceiverInternal(AActor* Receiver)
 	SendExtensionEventInternal(Receiver, NAME_ReceiverRemoved);
 }
 
-TSharedPtr<FComponentRequestHandle> UGameFrameworkComponentManager::AddComponentRequest(const TSoftClassPtr<AActor>& ReceiverClass, TSubclassOf<UActorComponent> ComponentClass)
+TSharedPtr<FComponentRequestHandle> UGameFrameworkComponentManager::AddComponentRequest(const TSoftClassPtr<AActor>& ReceiverClass, TSubclassOf<UActorComponent> ComponentClass, const EGameFrameworkAddComponentFlags AdditionFlags)
 {
 	// You must have a receiver and component class. The receiver cannot be AActor, that is too broad and would be bad for performance.
 	if (!ensure(!ReceiverClass.IsNull()) || !ensure(ComponentClass) || !ensure(ReceiverClass.ToString() != TEXT("/Script/Engine.Actor")))
@@ -253,14 +253,16 @@ TSharedPtr<FComponentRequestHandle> UGameFrameworkComponentManager::AddComponent
 	FComponentRequest NewRequest;
 	NewRequest.ReceiverClassPath = ReceiverClassPath;
 	NewRequest.ComponentClass = ComponentClassPtr;
+	
+	// Add a request if there is not an already existing one. Note that it will only uses the receiver and component class to check for uniqueness, not the addition flags.
 	int32& RequestCount = RequestTrackingMap.FindOrAdd(NewRequest);
 	RequestCount++;
 
 	if (RequestCount == 1)
 	{
-		auto& ComponentClasses = ReceiverClassToComponentClassMap.FindOrAdd(ReceiverClassPath);
-		ComponentClasses.Add(ComponentClassPtr);
-
+		EGameFrameworkAddComponentResult Result = EGameFrameworkAddComponentResult::Failed;
+		auto& RequestInfoSet = ReceiverClassToComponentClassMap.FindOrAdd(ReceiverClassPath);
+		RequestInfoSet.Add({ ComponentClassPtr, AdditionFlags } );
 		if (UClass* ReceiverClassPtr = ReceiverClass.Get())
 		{
 			UGameInstance* LocalGameInstance = GetGameInstance();
@@ -279,7 +281,7 @@ TSharedPtr<FComponentRequestHandle> UGameFrameworkComponentManager::AddComponent
 								ensureMsgf(AllReceivers.Contains(*ActorIt), TEXT("You may not add a component request for an actor class that does not call AddReceiver/RemoveReceiver in code! Class:%s"), *GetPathNameSafe(ReceiverClassPtr));
 							}
 #endif
-							CreateComponentOnInstance(*ActorIt, ComponentClass);
+							Result = CreateComponentOnInstance(*ActorIt, ComponentClass, AdditionFlags);
 						}
 					}
 				}
@@ -290,7 +292,10 @@ TSharedPtr<FComponentRequestHandle> UGameFrameworkComponentManager::AddComponent
 			// Actor class is not in memory, there will be no actor instances
 		}
 
-		return MakeShared<FComponentRequestHandle>(this, ReceiverClass, ComponentClass);
+		if (Result == EGameFrameworkAddComponentResult::Success)
+		{
+			return MakeShared<FComponentRequestHandle>(this, ReceiverClass, ComponentClass);
+		}
 	}
 
 	return nullptr;
@@ -310,10 +315,10 @@ void UGameFrameworkComponentManager::RemoveComponentRequest(const TSoftClassPtr<
 
 	if (RequestCount == 0)
 	{
-		if (auto* ComponentClasses = ReceiverClassToComponentClassMap.Find(ReceiverClassPath))
+		if (TSet<FComponentRequestInfo>* ReceiverSetInfo = ReceiverClassToComponentClassMap.Find(ReceiverClassPath))
 		{
-			ComponentClasses->Remove(ComponentClassPtr);
-			if (ComponentClasses->Num() == 0)
+			ReceiverSetInfo->Remove(ComponentClassPtr);
+			if (ReceiverSetInfo->Num() == 0)
 			{
 				ReceiverClassToComponentClassMap.Remove(ReceiverClassPath);
 			}
@@ -463,14 +468,42 @@ void UGameFrameworkComponentManager::SendExtensionEventInternal(AActor* Receiver
 	}
 }
 
-void UGameFrameworkComponentManager::CreateComponentOnInstance(AActor* ActorInstance, TSubclassOf<UActorComponent> ComponentClass)
+EGameFrameworkAddComponentResult UGameFrameworkComponentManager::CreateComponentOnInstance(AActor* ActorInstance, TSubclassOf<UActorComponent> ComponentClass, const EGameFrameworkAddComponentFlags AdditionFlags)
 {
 	check(ActorInstance);
 	check(ComponentClass);
 
 	if (!ComponentClass->GetDefaultObject<UActorComponent>()->GetIsReplicated() || ActorInstance->GetLocalRole() == ROLE_Authority)
 	{
-		UActorComponent* NewComp = NewObject<UActorComponent>(ActorInstance, ComponentClass, ComponentClass->GetFName());
+		// If AddUnique is set, it will be added only if no component on ActorInstance is child (or same class) of ComponentClass
+		const bool bAddUnique = EnumHasAnyFlags(AdditionFlags, EGameFrameworkAddComponentFlags::AddUnique);
+		if (bAddUnique)
+		{			
+			if (ActorInstance->GetComponentByClass(ComponentClass))
+			{
+				return EGameFrameworkAddComponentResult::Failed;
+			}
+		}
+
+		// If AddIfNotChild is set,it will be added only if ComponentClass is not a child (or same class) of an existing component on ActorInstance
+		const bool bAddIfNotChild = EnumHasAnyFlags(AdditionFlags, EGameFrameworkAddComponentFlags::AddIfNotChild);
+		if (bAddIfNotChild)
+		{
+			const TSet<UActorComponent*>& Components = ActorInstance->GetComponents();
+			for (const UActorComponent* ActorComp : Components)
+			{
+				if (ComponentClass->IsChildOf(ActorComp->GetClass()))
+				{
+					return EGameFrameworkAddComponentResult::Failed;
+				}
+			}
+		}
+
+		// If UseAutoGeneratedName is set, it will generate a new name and not re-use the class name directly (which can lead to component recycling)
+		const bool bUseAutoGeneratedName = EnumHasAnyFlags(AdditionFlags, EGameFrameworkAddComponentFlags::UseAutoGeneratedName);
+		const FName NewComponentName = bUseAutoGeneratedName ? NAME_None : ComponentClass->GetFName();
+
+		UActorComponent* NewComp = NewObject<UActorComponent>(ActorInstance, ComponentClass, NewComponentName);
 		TSet<FObjectKey>& ComponentInstances = ComponentClassToComponentInstanceMap.FindOrAdd(*ComponentClass);
 		ComponentInstances.Add(NewComp);
 
@@ -480,7 +513,11 @@ void UGameFrameworkComponentManager::CreateComponentOnInstance(AActor* ActorInst
 		}
 
 		NewComp->RegisterComponent();
+
+		return EGameFrameworkAddComponentResult::Success;
 	}
+
+	return EGameFrameworkAddComponentResult::Failed;
 }
 
 void UGameFrameworkComponentManager::DestroyInstancedComponent(UActorComponent* Component)
diff --git a/Engine/Plugins/Runtime/ModularGameplay/Source/ModularGameplay/Public/Components/GameFrameworkComponentManager.h b/Engine/Plugins/Runtime/ModularGameplay/Source/ModularGameplay/Public/Components/GameFrameworkComponentManager.h
--- a/Engine/Plugins/Runtime/ModularGameplay/Source/ModularGameplay/Public/Components/GameFrameworkComponentManager.h
+++ b/Engine/Plugins/Runtime/ModularGameplay/Source/ModularGameplay/Public/Components/GameFrameworkComponentManager.h
@@ -16,6 +16,23 @@ class UActorComponent;
 class UGameFrameworkComponent;
 
 
+// Bitflag enum to control how this AddComponents action should actually run
+UENUM(BlueprintType, meta = (Bitflags, UseEnumValuesAsMaskValuesInEditor = "true"))
+enum class EGameFrameworkAddComponentFlags : uint8
+{
+	// If 0, no specific rule is applied and the component will be always created on the receiver"),
+	None = 0 UMETA(Hidden),
+	AddUnique = 0b00000001				UMETA(DisplayName = "Add Unique", ToolTip="If AddUnique is set, it will be added only if there is no existing component on ActorInstance of ComponentClass"),
+	AddIfNotChild = 0b00000010			UMETA(DisplayName = "Add Unique If Not Child", ToolTip = "If Add Unique If Not Child is set, it will be added only if ComponentClass is not a child (or same class) of an existing component on ActorInstance"),
+	UseAutoGeneratedName = 0b00000100	UMETA(DisplayName = "Use Auto Generated Name", ToolTip = "If UseAutoGeneratedName is set, it will generate a new name and not re-use the class name directly (which can lead to component recycling)."),
+};
+ENUM_CLASS_FLAGS(EGameFrameworkAddComponentFlags);
+
+enum class EGameFrameworkAddComponentResult : uint8
+{
+	Success,
+	Failed
+};
 /** 
  * A handle for a request to put components or call a delegate for an extensible actor class.
  * When this handle is destroyed, it will remove the associated request from the system.
@@ -99,8 +116,16 @@ public:
 	/** Removes an actor as a receiver for components (automatically finding the manager for the actor's game instance). */
 	static void RemoveGameFrameworkComponentReceiver(AActor* Receiver);
 
-	/** Adds a request to instantiate components on actors of the given classes. Returns a handle that will keep the request "alive" until it is destructed, at which point the request is removed. */
-	TSharedPtr<FComponentRequestHandle> AddComponentRequest(const TSoftClassPtr<AActor>& ReceiverClass, TSubclassOf<UActorComponent> ComponentClass);
+	/** 
+	* Adds a request to instantiate components on actors of the given classes. 
+	* 
+	* @param ReceiverClass		The class on which we should add the component onto
+	* @param ComponentClass		The component class to add on the new instance
+	* @param AdditionFlags (optional)	Set of flags to considered at component creation
+	* 
+	* @return Returns a handle that will keep the request "alive" until it is destructed, at which point the request is removed. Will return null if there is already an active request with same Receiver & Component classes.
+	*/
+	TSharedPtr<FComponentRequestHandle> AddComponentRequest(const TSoftClassPtr<AActor>& ReceiverClass, TSubclassOf<UActorComponent> ComponentClass, const EGameFrameworkAddComponentFlags AdditionFlags = EGameFrameworkAddComponentFlags::None);
 
 
 	//////////////////////////////////////////////////////////////////////////////////////////////
@@ -158,11 +183,36 @@ private:
 	void RemoveExtensionHandler(const TSoftClassPtr<AActor>& ReceiverClass, FDelegateHandle DelegateHandle);
 
 	/** Creates an instance of a component on an actor */
-	void CreateComponentOnInstance(AActor* ActorInstance, TSubclassOf<UActorComponent> ComponentClass);
+	EGameFrameworkAddComponentResult CreateComponentOnInstance(AActor* ActorInstance, TSubclassOf<UActorComponent> ComponentClass, const EGameFrameworkAddComponentFlags AdditionFlags = EGameFrameworkAddComponentFlags::None);
 
 	/** Removes an instance of a component on an actor */
 	void DestroyInstancedComponent(UActorComponent* Component);
 
+	/** A pair holding a receiver class and AdditionFlags */
+	struct FComponentRequestInfo
+	{
+		FComponentRequestInfo(UClass* InClass, const EGameFrameworkAddComponentFlags InAdditionFlags = EGameFrameworkAddComponentFlags::None)
+			: Class(InClass)
+			, AdditionFlags(InAdditionFlags)
+		{}
+
+		bool operator==(const FComponentRequestInfo& Other) const
+		{
+			return Class == Other.Class;
+		}
+
+		friend FORCEINLINE uint32 GetTypeHash(const FComponentRequestInfo& Request)
+		{
+			return GetTypeHash(Request.Class);
+		}
+
+		// The receiver class
+		TObjectPtr<UClass> Class;
+
+		// Observe these rules when adding the component, if any
+		EGameFrameworkAddComponentFlags AdditionFlags = EGameFrameworkAddComponentFlags::None;
+	};
+
 	/** A list of FNames to represent an object path. Used for fast hashing and comparison of paths */
 	struct FComponentRequestReceiverClassPath
 	{
@@ -247,8 +297,8 @@ private:
 	/** A map of component classes to instances of that component class made by this component manager */
 	TMap<UClass*, TSet<FObjectKey>> ComponentClassToComponentInstanceMap;
 
-	/** A map of actor classes to component classes that should be made for that class. */
-	TMap<FComponentRequestReceiverClassPath, TSet<TObjectPtr<UClass>>> ReceiverClassToComponentClassMap;
+	/** A map of actor classes to component classes & addition flags that should be made for that class. */
+	TMap<FComponentRequestReceiverClassPath, TSet<FComponentRequestInfo>> ReceiverClassToComponentClassMap;
 
 	/** A map of actor classes to delegate handlers that should be executed for actors of that class. */
 	TMap<FComponentRequestReceiverClassPath, FExtensionHandlerEvent> ReceiverClassToEventMap;
