ModularGameplay: Fixed UGameFrameworkComponentManager holding on to references.
#rnx

diff --git a/Engine/Plugins/Runtime/ModularGameplay/Source/ModularGameplay/Private/Components/GameFrameworkComponentManager.cpp b/Engine/Plugins/Runtime/ModularGameplay/Source/ModularGameplay/Private/Components/GameFrameworkComponentManager.cpp
--- a/Engine/Plugins/Runtime/ModularGameplay/Source/ModularGameplay/Private/Components/GameFrameworkComponentManager.cpp
+++ b/Engine/Plugins/Runtime/ModularGameplay/Source/ModularGameplay/Private/Components/GameFrameworkComponentManager.cpp
@@ -292,10 +292,7 @@ TSharedPtr<FComponentRequestHandle> UGameFrameworkComponentManager::AddComponent
 			// Actor class is not in memory, there will be no actor instances
 		}
 
-		if (Result == EGameFrameworkAddComponentResult::Success)
-		{
-			return MakeShared<FComponentRequestHandle>(this, ReceiverClass, ComponentClass);
-		}
+		return MakeShared<FComponentRequestHandle>(this, ReceiverClass, ComponentClass);
 	}
 
 	return nullptr;
