PR #11514: Add get debug state virtual to aid debugging of the init state interface

#pf 65fd98130f0b188ba69a46d8

diff --git a/Engine/Plugins/Runtime/ModularGameplay/Source/ModularGameplay/Private/Components/GameFrameworkInitStateInterface.cpp b/Engine/Plugins/Runtime/ModularGameplay/Source/ModularGameplay/Private/Components/GameFrameworkInitStateInterface.cpp
--- a/Engine/Plugins/Runtime/ModularGameplay/Source/ModularGameplay/Private/Components/GameFrameworkInitStateInterface.cpp
+++ b/Engine/Plugins/Runtime/ModularGameplay/Source/ModularGameplay/Private/Components/GameFrameworkInitStateInterface.cpp
@@ -256,3 +256,7 @@ bool IGameFrameworkInitStateInterface::UnregisterInitStateDelegate(FActorInitSta
 	return false;
 }
 
+FString IGameFrameworkInitStateInterface::GetDebugState() const
+{
+	return FString::Printf(TEXT("CurrentState: %s"), *GetInitState().ToString());
+}
diff --git a/Engine/Plugins/Runtime/ModularGameplay/Source/ModularGameplay/Public/Components/GameFrameworkInitStateInterface.h b/Engine/Plugins/Runtime/ModularGameplay/Source/ModularGameplay/Public/Components/GameFrameworkInitStateInterface.h
--- a/Engine/Plugins/Runtime/ModularGameplay/Source/ModularGameplay/Public/Components/GameFrameworkInitStateInterface.h
+++ b/Engine/Plugins/Runtime/ModularGameplay/Source/ModularGameplay/Public/Components/GameFrameworkInitStateInterface.h
@@ -79,6 +79,9 @@ public:
 	UFUNCTION(BlueprintCallable, Category = "InitState")
 	virtual bool UnregisterInitStateDelegate(FActorInitStateChangedBPDelegate Delegate);
 
+	/** Returns Current state and any additional debug information for the active state */
+	virtual FString GetDebugState() const;
+
 protected:
 	/** Default handle created from calling BindOnActorInitStateChanged */
 	FDelegateHandle ActorInitStateChangedHandle;
