Fix safety issues with modifying GameFrameworkComponentManager during delegate execution by storing the delegates as shared refs.
This fixes memory crashes from resizing containers during execution and properly calls all delegates if earlier ones are removed during iteration.
This also stops double execution of newly added delegates and does fewer delegate copies.

diff --git a/Engine/Plugins/Runtime/ModularGameplay/Source/ModularGameplay/Private/Components/GameFrameworkComponentManager.cpp b/Engine/Plugins/Runtime/ModularGameplay/Source/ModularGameplay/Private/Components/GameFrameworkComponentManager.cpp
--- a/Engine/Plugins/Runtime/ModularGameplay/Source/ModularGameplay/Private/Components/GameFrameworkComponentManager.cpp
+++ b/Engine/Plugins/Runtime/ModularGameplay/Source/ModularGameplay/Private/Components/GameFrameworkComponentManager.cpp
@@ -190,9 +190,11 @@ void UGameFrameworkComponentManager::AddReceiverInternal(AActor* Receiver)
 
 		if (FExtensionHandlerEvent* HandlerEvent = ReceiverClassToEventMap.Find(ReceiverClassPath))
 		{
-			for (const TPair<FDelegateHandle, FExtensionHandlerDelegate>& Pair : *HandlerEvent)
+			// Copy the execution list so it isn't invalidated
+			FExtensionHandlerEvent HandlerCopy = *HandlerEvent;
+			for (const TPair<FDelegateHandle, TSharedRef<FExtensionHandlerRegisteredDelegate>>& Pair : HandlerCopy)
 			{
-				Pair.Value.Execute(Receiver, NAME_ReceiverAdded);
+				Pair.Value->Execute(Receiver, NAME_ReceiverAdded);
 			}
 		}
 	}
@@ -359,9 +361,11 @@ TSharedPtr<FComponentRequestHandle> UGameFrameworkComponentManager::AddExtension
 	FComponentRequestReceiverClassPath ReceiverClassPath(ReceiverClass);
 	FExtensionHandlerEvent& HandlerEvent = ReceiverClassToEventMap.FindOrAdd(ReceiverClassPath);
 
-	// This is a fake multicast delegate using a map
+	// This is a fake multicast delegate using a map, we store it in shared memory to avoid things being reallocated during execution
+	// This avoids copying the delegate itself which is often expensive
 	FDelegateHandle DelegateHandle(FDelegateHandle::EGenerateNewHandleType::GenerateNewHandle);
-	HandlerEvent.Add(DelegateHandle, ExtensionHandler);
+	TSharedRef<FExtensionHandlerRegisteredDelegate> RegisteredHandler = MakeShared<FExtensionHandlerRegisteredDelegate>(MoveTemp(ExtensionHandler));
+	HandlerEvent.Add(DelegateHandle, RegisteredHandler);
 
 	if (UClass* ReceiverClassPtr = ReceiverClass.Get())
 	{
@@ -375,7 +379,7 @@ TSharedPtr<FComponentRequestHandle> UGameFrameworkComponentManager::AddExtension
 				{
 					if (ActorIt->IsActorInitialized())
 					{
-						ExtensionHandler.Execute(*ActorIt, NAME_ExtensionAdded);
+						RegisteredHandler->Execute(*ActorIt, NAME_ExtensionAdded);
 					}
 				}
 			}
@@ -395,9 +399,10 @@ void UGameFrameworkComponentManager::RemoveExtensionHandler(const TSoftClassPtr<
 
 	if (FExtensionHandlerEvent* HandlerEvent = ReceiverClassToEventMap.Find(ReceiverClassPath))
 	{
-		FExtensionHandlerDelegate* HandlerDelegate = HandlerEvent->Find(DelegateHandle);
+		TSharedRef<FExtensionHandlerRegisteredDelegate>* HandlerDelegate = HandlerEvent->Find(DelegateHandle);
 		if (ensure(HandlerDelegate))
 		{
+			TSharedRef<FExtensionHandlerRegisteredDelegate> HandlerRef = *HandlerDelegate;
 			// Call it once on unregister
 			if (UClass* ReceiverClassPtr = ReceiverClass.Get())
 			{
@@ -412,7 +417,7 @@ void UGameFrameworkComponentManager::RemoveExtensionHandler(const TSoftClassPtr<
 						{
 							if (ActorIt->IsActorInitialized())
 							{
-								HandlerDelegate->Execute(*ActorIt, NAME_ExtensionRemoved);
+								HandlerRef->Execute(*ActorIt, NAME_ExtensionRemoved);
 							}
 						}
 					}
@@ -423,6 +428,8 @@ void UGameFrameworkComponentManager::RemoveExtensionHandler(const TSoftClassPtr<
 				// Actor class is not in memory, there will be no actor instances
 			}
 
+			// This stops it from executing further up the stack
+			HandlerRef->bRemoved = true;
 			HandlerEvent->Remove(DelegateHandle);
 
 			if (HandlerEvent->IsEmpty())
@@ -457,9 +464,11 @@ void UGameFrameworkComponentManager::SendExtensionEventInternal(AActor* Receiver
 		FComponentRequestReceiverClassPath ReceiverClassPath(Class);
 		if (FExtensionHandlerEvent* HandlerEvent = ReceiverClassToEventMap.Find(ReceiverClassPath))
 		{
-			for (const TPair<FDelegateHandle, FExtensionHandlerDelegate>& Pair : *HandlerEvent)
+			// Copy the execution list so it isn't invalidated
+			FExtensionHandlerEvent HandlerCopy = *HandlerEvent;
+			for (const TPair<FDelegateHandle, TSharedRef<FExtensionHandlerRegisteredDelegate>>& Pair : HandlerCopy)
 			{
-				Pair.Value.Execute(Receiver, EventName);
+				Pair.Value->Execute(Receiver, EventName);
 			}
 		}
 	}
@@ -563,6 +572,7 @@ UGameFrameworkComponentManager::FActorFeatureRegisteredDelegate::FActorFeatureRe
 	, DelegateHandle(FDelegateHandle::EGenerateNewHandleType::GenerateNewHandle)
 	, RequiredFeatureName(InFeatureName)
 	, RequiredInitState(InInitState)
+	, bRemoved(false)
 {
 
 }
@@ -572,12 +582,18 @@ UGameFrameworkComponentManager::FActorFeatureRegisteredDelegate::FActorFeatureRe
 	, DelegateHandle(FDelegateHandle::EGenerateNewHandleType::GenerateNewHandle)
 	, RequiredFeatureName(InFeatureName)
 	, RequiredInitState(InInitState)
+	, bRemoved(false)
 {
 
 }
 
 void UGameFrameworkComponentManager::FActorFeatureRegisteredDelegate::Execute(AActor* OwningActor, FName FeatureName, UObject* Implementer, FGameplayTag FeatureState)
 {
+	if (bRemoved)
+	{
+		return;
+	}
+
 	FActorInitStateChangedParams Params(OwningActor, FeatureName, Implementer, FeatureState);
 	if (Delegate.IsBound())
 	{
@@ -591,6 +607,15 @@ void UGameFrameworkComponentManager::FActorFeatureRegisteredDelegate::Execute(AA
 	}
 }
 
+UGameFrameworkComponentManager::FActorFeatureData::~FActorFeatureData()
+{
+	for (TSharedRef<FActorFeatureRegisteredDelegate>& RegisteredDelegate : RegisteredDelegates)
+	{
+		// This ensures it will not execute if cleared during delegate iteration
+		RegisteredDelegate->bRemoved = true;
+	}
+}
+
 void UGameFrameworkComponentManager::RegisterInitState(FGameplayTag NewState, bool bAddBefore, FGameplayTag ExistingState)
 {
 	int32 FoundIndex;
@@ -826,18 +851,16 @@ FDelegateHandle UGameFrameworkComponentManager::RegisterAndCallForActorInitState
 		// We often register delegates before registering states
 		FActorFeatureData& ActorStruct = FindOrAddActorData(Actor);
 
-		FActorFeatureRegisteredDelegate& RegisteredDelegate = ActorStruct.RegisteredDelegates.Emplace_GetRef(MoveTemp(Delegate), FeatureName, RequiredState);
-
-		// Cache handle as delegate could invalidate it
-		FDelegateHandle ReturnHandle = RegisteredDelegate.DelegateHandle;
+		TSharedRef<FActorFeatureRegisteredDelegate> RegisteredDelegate = MakeShared<FActorFeatureRegisteredDelegate>(MoveTemp(Delegate), FeatureName, RequiredState);
+		ActorStruct.RegisteredDelegates.Add(RegisteredDelegate);
 
 		if (bCallImmediately)
 		{
-			FActorFeatureRegisteredDelegate DelegateCopy = RegisteredDelegate;
-			CallDelegateForMatchingFeatures(Actor, DelegateCopy);
+			// The shared ref keeps the delegate alive in case it gets unregistered
+			CallDelegateForMatchingFeatures(Actor, *RegisteredDelegate);
 		}		
 		
-		return ReturnHandle;
+		return RegisteredDelegate->DelegateHandle;
 	}
 
 	return FDelegateHandle();
@@ -850,15 +873,13 @@ bool UGameFrameworkComponentManager::RegisterAndCallForActorInitState(AActor* Ac
 		// We often register delegates before registering states
 		FActorFeatureData& ActorStruct = FindOrAddActorData(Actor);
 
-		FActorFeatureRegisteredDelegate& RegisteredDelegate = ActorStruct.RegisteredDelegates.Emplace_GetRef(MoveTemp(Delegate), FeatureName, RequiredState);
-
-		// Cache handle as delegate could invalidate it
-		FDelegateHandle ReturnHandle = RegisteredDelegate.DelegateHandle;
+		TSharedRef<FActorFeatureRegisteredDelegate> RegisteredDelegate = MakeShared<FActorFeatureRegisteredDelegate>(MoveTemp(Delegate), FeatureName, RequiredState);
+		ActorStruct.RegisteredDelegates.Add(RegisteredDelegate);
 
 		if (bCallImmediately)
 		{
-			FActorFeatureRegisteredDelegate DelegateCopy = RegisteredDelegate;
-			CallDelegateForMatchingFeatures(Actor, DelegateCopy);
+			// The shared ref keeps the delegate alive in case it gets unregistered
+			CallDelegateForMatchingFeatures(Actor, *RegisteredDelegate);
 		}
 
 		return true;
@@ -875,14 +896,7 @@ bool UGameFrameworkComponentManager::UnregisterActorInitStateDelegate(AActor* Ac
 
 		if (ActorStruct)
 		{
-			int32 FoundIndex = GetIndexForRegisteredDelegate(ActorStruct->RegisteredDelegates, Handle);
-
-			if (FoundIndex >= 0)
-			{
-				Handle.Reset();
-				ActorStruct->RegisteredDelegates.RemoveAt(FoundIndex);
-				return true;
-			}
+			return RemoveActorFeatureDelegateFromList(ActorStruct->RegisteredDelegates, Handle);
 		}
 	}
 
@@ -897,13 +911,7 @@ bool UGameFrameworkComponentManager::UnregisterActorInitStateDelegate(AActor* Ac
 
 		if (ActorStruct)
 		{
-			int32 FoundIndex = GetIndexForRegisteredDelegate(ActorStruct->RegisteredDelegates, DelegateToRemove);
-
-			if (FoundIndex >= 0)
-			{
-				ActorStruct->RegisteredDelegates.RemoveAt(FoundIndex);
-				return true;
-			}
+			return RemoveActorFeatureDelegateFromList(ActorStruct->RegisteredDelegates, DelegateToRemove);
 		}
 	}
 
@@ -915,23 +923,18 @@ FDelegateHandle UGameFrameworkComponentManager::RegisterAndCallForClassInitState
 	if (ensure(!ActorClass.IsNull() && Delegate.IsBound() && !FeatureName.IsNone()))
 	{
 		FComponentRequestReceiverClassPath ReceiverClassPath(ActorClass);
-		TArray<FActorFeatureRegisteredDelegate>& RegisteredDelegates = ClassFeatureChangeDelegates.FindOrAdd(ReceiverClassPath);
-
-		FActorFeatureRegisteredDelegate& RegisteredDelegate = RegisteredDelegates.Emplace_GetRef(MoveTemp(Delegate), FeatureName, RequiredState);
-
-		// Cache handle as delegate could invalidate it
-		FDelegateHandle ReturnHandle = RegisteredDelegate.DelegateHandle;
+		FActorFeatureDelegateList& RegisteredDelegates = ClassFeatureChangeDelegates.FindOrAdd(ReceiverClassPath);
 
+		TSharedRef<FActorFeatureRegisteredDelegate> RegisteredDelegate = MakeShared<FActorFeatureRegisteredDelegate>(MoveTemp(Delegate), FeatureName, RequiredState);
+		RegisteredDelegates.Add(RegisteredDelegate);
+	
 		if (bCallImmediately)
 		{
-			// If RealClass isn't valid yet, there can't be any instances
-			UClass* RealClass = ActorClass.Get();
-
-			FActorFeatureRegisteredDelegate DelegateCopy = RegisteredDelegate;
-			CallDelegateForMatchingActors(RealClass, DelegateCopy);
+			// A null actor class means there are no registered instances and the call is ignored
+			CallDelegateForMatchingActors(ActorClass.Get(), *RegisteredDelegate);
 		}
 
-		return ReturnHandle;
+		return RegisteredDelegate->DelegateHandle;
 	}
 
 	return FDelegateHandle();
@@ -942,20 +945,15 @@ bool UGameFrameworkComponentManager::RegisterAndCallForClassInitState(TSoftClass
 	if (ensure(!ActorClass.IsNull() && Delegate.IsBound() && !FeatureName.IsNone()))
 	{
 		FComponentRequestReceiverClassPath ReceiverClassPath(ActorClass);
-		TArray<FActorFeatureRegisteredDelegate>& RegisteredDelegates = ClassFeatureChangeDelegates.FindOrAdd(ReceiverClassPath);
-
-		FActorFeatureRegisteredDelegate& RegisteredDelegate = RegisteredDelegates.Emplace_GetRef(MoveTemp(Delegate), FeatureName, RequiredState);
+		FActorFeatureDelegateList& RegisteredDelegates = ClassFeatureChangeDelegates.FindOrAdd(ReceiverClassPath);
 
-		// Cache handle as delegate could invalidate it
-		FDelegateHandle ReturnHandle = RegisteredDelegate.DelegateHandle;
+		TSharedRef<FActorFeatureRegisteredDelegate> RegisteredDelegate = MakeShared<FActorFeatureRegisteredDelegate>(MoveTemp(Delegate), FeatureName, RequiredState);
+		RegisteredDelegates.Add(RegisteredDelegate);
 
 		if (bCallImmediately)
 		{
-			// If RealClass isn't valid yet, there can't be any instances
-			UClass* RealClass = ActorClass.Get();
-
-			FActorFeatureRegisteredDelegate DelegateCopy = RegisteredDelegate;
-			CallDelegateForMatchingActors(RealClass, DelegateCopy);
+			// A null actor class means there are no registered instances and the call is ignored
+			CallDelegateForMatchingActors(ActorClass.Get(), *RegisteredDelegate);
 		}
 
 		return true;
@@ -969,18 +967,11 @@ bool UGameFrameworkComponentManager::UnregisterClassInitStateDelegate(const TSof
 	if (!ActorClass.IsNull() && Handle.IsValid())
 	{
 		FComponentRequestReceiverClassPath ReceiverClassPath(ActorClass);
-		TArray<FActorFeatureRegisteredDelegate>* RegisteredDelegates = ClassFeatureChangeDelegates.Find(ReceiverClassPath);
+		FActorFeatureDelegateList* RegisteredDelegates = ClassFeatureChangeDelegates.Find(ReceiverClassPath);
 
 		if (RegisteredDelegates)
 		{
-			int32 FoundIndex = GetIndexForRegisteredDelegate(*RegisteredDelegates, Handle);
-
-			if (FoundIndex >= 0)
-			{
-				Handle.Reset();
-				RegisteredDelegates->RemoveAt(FoundIndex);
-				return true;
-			}
+			return RemoveActorFeatureDelegateFromList(*RegisteredDelegates, Handle);
 		}
 	}
 
@@ -992,24 +983,17 @@ bool UGameFrameworkComponentManager::UnregisterClassInitStateDelegate(TSoftClass
 	if (!ActorClass.IsNull() && DelegateToRemove.IsBound())
 	{
 		FComponentRequestReceiverClassPath ReceiverClassPath(ActorClass);
-		TArray<FActorFeatureRegisteredDelegate>* RegisteredDelegates = ClassFeatureChangeDelegates.Find(ReceiverClassPath);
+		FActorFeatureDelegateList* RegisteredDelegates = ClassFeatureChangeDelegates.Find(ReceiverClassPath);
 
 		if (RegisteredDelegates)
 		{
-			int32 FoundIndex = GetIndexForRegisteredDelegate(*RegisteredDelegates, DelegateToRemove);
-
-			if (FoundIndex >= 0)
-			{
-				RegisteredDelegates->RemoveAt(FoundIndex);
-				return true;
-			}
+			return RemoveActorFeatureDelegateFromList(*RegisteredDelegates, DelegateToRemove);
 		}
 	}
 
 	return false;
 }
 
-
 const UGameFrameworkComponentManager::FActorFeatureState* UGameFrameworkComponentManager::FindFeatureStateStruct(const FActorFeatureData* ActorStruct, FName FeatureName, FGameplayTag RequiredState) const
 {
 	if (ActorStruct)
@@ -1053,80 +1037,65 @@ void UGameFrameworkComponentManager::ProcessFeatureStateChange(AActor* Actor, co
 void UGameFrameworkComponentManager::CallFeatureStateDelegates(AActor* Actor, FActorFeatureState StateChange)
 {
 	FActorFeatureData* ActorStruct = ActorFeatureMap.Find(FObjectKey(Actor));
-	int32 DelegateIndex = 0;
-	bool bMoreActorDelegates = true;
-	bool bMoreClassDelegates = true;
-	UClass* ClassToCheck = Actor->GetClass();
+	FActorFeatureDelegateList QueuedDelegates;
 
-	while (ActorStruct)
+	// Should only be called inside ProcessFeatureStateChange
+	ensure(CurrentStateChange != INDEX_NONE);
+
+	if (ActorStruct)
 	{
-		if (bMoreActorDelegates)
+		for (TSharedRef<FActorFeatureRegisteredDelegate>& DelegateRef : ActorStruct->RegisteredDelegates)
 		{
-			// First check the actor-specific delegates in order
-
-			if (ActorStruct->RegisteredDelegates.IsValidIndex(DelegateIndex))
-			{
-				FActorFeatureRegisteredDelegate& RegisteredDelegate = ActorStruct->RegisteredDelegates[DelegateIndex];
-
-				if ((RegisteredDelegate.RequiredFeatureName.IsNone() || RegisteredDelegate.RequiredFeatureName == StateChange.FeatureName)
-					&& (!RegisteredDelegate.RequiredInitState.IsValid() || IsInitStateAfterOrEqual(StateChange.CurrentState, RegisteredDelegate.RequiredInitState)))
-				{
-					// TODO remove invalid delegates now?
-					RegisteredDelegate.Execute(Actor, StateChange.FeatureName, StateChange.Implementer.Get(), StateChange.CurrentState);
-
-					// That could have invalidated anything
-					ActorStruct = ActorFeatureMap.Find(FObjectKey(Actor));
-				}
-
-				DelegateIndex++;
-			}
-			else
+			FActorFeatureRegisteredDelegate& RegisteredDelegate = *DelegateRef;
+			if ((RegisteredDelegate.RequiredFeatureName.IsNone() || RegisteredDelegate.RequiredFeatureName == StateChange.FeatureName)
+				&& (!RegisteredDelegate.RequiredInitState.IsValid() || IsInitStateAfterOrEqual(StateChange.CurrentState, RegisteredDelegate.RequiredInitState)))
 			{
-				DelegateIndex = 0;
-				bMoreActorDelegates = false;
+				// Queue delegates now in case the registered list changes during execution
+				// If new delegates are registered, they are handled at registration time if bCallImmediately is used
+				QueuedDelegates.Add(DelegateRef);
 			}
 		}
-		else if (ClassToCheck)
+
+		UClass* ClassToCheck = Actor->GetClass();
+		while (ClassToCheck)
 		{
 			// Now check the general class delegates
-
 			FComponentRequestReceiverClassPath ReceiverClassPath(ClassToCheck);
-			TArray<FActorFeatureRegisteredDelegate>* FoundDelegates = ClassFeatureChangeDelegates.Find(ReceiverClassPath);
+			FActorFeatureDelegateList* FoundDelegates = ClassFeatureChangeDelegates.Find(ReceiverClassPath);
 
-			if (FoundDelegates && FoundDelegates->IsValidIndex(DelegateIndex))
+			if (FoundDelegates)
 			{
-				FActorFeatureRegisteredDelegate& RegisteredDelegate = (*FoundDelegates)[DelegateIndex];
-
-				if ((RegisteredDelegate.RequiredFeatureName.IsNone() || RegisteredDelegate.RequiredFeatureName == StateChange.FeatureName)
-					&& (!RegisteredDelegate.RequiredInitState.IsValid() || IsInitStateAfterOrEqual(StateChange.CurrentState, RegisteredDelegate.RequiredInitState)))
+				for (TSharedRef<FActorFeatureRegisteredDelegate>& DelegateRef : *FoundDelegates)
 				{
-					// TODO remove invalid delegates now?
-					RegisteredDelegate.Execute(Actor, StateChange.FeatureName, StateChange.Implementer.Get(), StateChange.CurrentState);
+					FActorFeatureRegisteredDelegate& RegisteredDelegate = *DelegateRef;
+					if ((RegisteredDelegate.RequiredFeatureName.IsNone() || RegisteredDelegate.RequiredFeatureName == StateChange.FeatureName)
+						&& (!RegisteredDelegate.RequiredInitState.IsValid() || IsInitStateAfterOrEqual(StateChange.CurrentState, RegisteredDelegate.RequiredInitState)))
+					{
+						QueuedDelegates.Add(DelegateRef);
+					}
 
-					// That could have invalidated anything
-					ActorStruct = ActorFeatureMap.Find(FObjectKey(Actor));
 				}
-
-				DelegateIndex++;
 			}
-			else
-			{
-				// Start over with parent class
-				DelegateIndex = 0;
-				ClassToCheck = ClassToCheck->GetSuperClass();
-			}
-		}
-		else
-		{
-			return;
+
+			ClassToCheck = ClassToCheck->GetSuperClass();
 		}
 	}
+
+	// Now execute the queued delegates, if they are removed Execute will skip them
+	for (TSharedRef<FActorFeatureRegisteredDelegate>& QueuedDelegate : QueuedDelegates)
+	{
+		QueuedDelegate->Execute(Actor, StateChange.FeatureName, StateChange.Implementer.Get(), StateChange.CurrentState);
+	}
 }
 
 void UGameFrameworkComponentManager::CallDelegateForMatchingFeatures(AActor* Actor, FActorFeatureRegisteredDelegate& RegisteredDelegate)
 {
 	FActorFeatureData* ActorStruct = ActorFeatureMap.Find(FObjectKey(Actor));
-	int32 StateIndex = 0;
+
+	if (ActorStruct == nullptr)
+	{
+		return;
+	}
 
 	// If feature is specified, just call the one
 	if (!RegisteredDelegate.RequiredFeatureName.IsNone())
@@ -1141,20 +1110,20 @@ void UGameFrameworkComponentManager::CallDelegateForMatchingFeatures(AActor* Act
 		return;
 	}
 	
-	// If feature is not specified, iterate them all
-	while (ActorStruct && ActorStruct->RegisteredStates.IsValidIndex(StateIndex))
+	// If feature is not specified, iterate and run on all valid ones
+	TArray<FActorFeatureState> QueuedStates;
+	for (FActorFeatureState& FeatureState : ActorStruct->RegisteredStates)
 	{
-		const FActorFeatureState* FoundStruct = &ActorStruct->RegisteredStates[StateIndex];
-
-		if (!RegisteredDelegate.RequiredInitState.IsValid() || IsInitStateAfterOrEqual(FoundStruct->CurrentState, RegisteredDelegate.RequiredInitState))
+		if (!RegisteredDelegate.RequiredInitState.IsValid() || IsInitStateAfterOrEqual(FeatureState.CurrentState, RegisteredDelegate.RequiredInitState))
 		{
-			RegisteredDelegate.Execute(Actor, FoundStruct->FeatureName, FoundStruct->Implementer.Get(), FoundStruct->CurrentState);
-
-			// That could have invalidated anything
-			ActorStruct = ActorFeatureMap.Find(FObjectKey(Actor));
+			// Make a copy in case the state memory is invalidated
+			QueuedStates.Add(FeatureState);
 		}
+	}
 
-		StateIndex++;
+	for (FActorFeatureState& FeatureState : QueuedStates)
+	{
+		RegisteredDelegate.Execute(Actor, FeatureState.FeatureName, FeatureState.Implementer.Get(), FeatureState.CurrentState);
 	}
 }
 
@@ -1200,28 +1169,33 @@ UGameFrameworkComponentManager::FActorFeatureData& UGameFrameworkComponentManage
 	return ActorStruct;
 }
 
-int32 UGameFrameworkComponentManager::GetIndexForRegisteredDelegate(TArray<FActorFeatureRegisteredDelegate>& DelegatesToSearch, FDelegateHandle SearchHandle) const
+bool UGameFrameworkComponentManager::RemoveActorFeatureDelegateFromList(FActorFeatureDelegateList& DelegateList, FDelegateHandle& SearchHandle) const
 {
-	for (int32 i = DelegatesToSearch.Num() - 1; i >= 0; i--)
+	for (int32 i = DelegateList.Num() - 1; i >= 0; i--)
 	{
-		if (DelegatesToSearch[i].DelegateHandle == SearchHandle)
+		if (DelegateList[i]->DelegateHandle == SearchHandle)
 		{
-			return i;
+			DelegateList[i]->bRemoved = true;
+			DelegateList.RemoveAt(i);
+			SearchHandle.Reset();
+			return true;
 		}
 	}
 
-	return INDEX_NONE;
+	return false;
 }
 
-int32 UGameFrameworkComponentManager::GetIndexForRegisteredDelegate(TArray<FActorFeatureRegisteredDelegate>& DelegatesToSearch, FActorInitStateChangedBPDelegate SearchDelegate) const
+bool UGameFrameworkComponentManager::RemoveActorFeatureDelegateFromList(FActorFeatureDelegateList& DelegateList, FActorInitStateChangedBPDelegate SearchDelegate) const
 {
-	for (int32 i = DelegatesToSearch.Num() - 1; i >= 0; i--)
+	for (int32 i = DelegateList.Num() - 1; i >= 0; i--)
 	{
-		if (DelegatesToSearch[i].BPDelegate == SearchDelegate)
+		if (DelegateList[i]->BPDelegate == SearchDelegate)
 		{
-			return i;
+			DelegateList[i]->bRemoved = true;
+			DelegateList.RemoveAt(i);
+			return true;
 		}
 	}
 
-	return INDEX_NONE;
+	return false;
 }
diff --git a/Engine/Plugins/Runtime/ModularGameplay/Source/ModularGameplay/Public/Components/GameFrameworkComponentManager.h b/Engine/Plugins/Runtime/ModularGameplay/Source/ModularGameplay/Public/Components/GameFrameworkComponentManager.h
--- a/Engine/Plugins/Runtime/ModularGameplay/Source/ModularGameplay/Public/Components/GameFrameworkComponentManager.h
+++ b/Engine/Plugins/Runtime/ModularGameplay/Source/ModularGameplay/Public/Components/GameFrameworkComponentManager.h
@@ -89,12 +89,9 @@ class MODULARGAMEPLAY_API UGameFrameworkComponentManager : public UGameInstanceS
 {
 	GENERATED_BODY()
 
-	/** Using a fake multicast delegate so order can be kept consistent */
-	DECLARE_DELEGATE_TwoParams(FExtensionHandlerDelegateInternal, AActor*, FName);
-	using FExtensionHandlerEvent = TMap<FDelegateHandle, FExtensionHandlerDelegateInternal>;
 public:
 	/** Delegate types for extension handlers */
-	using FExtensionHandlerDelegate = FExtensionHandlerDelegateInternal;
+	DECLARE_DELEGATE_TwoParams(FExtensionHandlerDelegate, AActor*, FName);
 
 	virtual void Initialize(FSubsystemCollectionBase& Collection) override;
 	virtual void Deinitialize() override;
@@ -291,6 +288,35 @@ private:
 		}
 	};
 
+	/** Stores a delegate as shared ref to avoid deallocation during execution */
+	struct FExtensionHandlerRegisteredDelegate
+	{
+		/** Delegate to call */
+		FExtensionHandlerDelegate Delegate;
+
+		/** Has this been removed, if so do not call */
+		bool bRemoved;
+
+		/** Construct from an existing delegate, there is no default constructor */
+		FExtensionHandlerRegisteredDelegate(FExtensionHandlerDelegate&& InDelegate)
+			: Delegate(InDelegate)
+			, bRemoved(false)
+		{
+		}
+
+		/** Call Delegate if it is valid */
+		void Execute(AActor* OwningActor, FName EventName)
+		{
+			if (!bRemoved)
+			{
+				Delegate.Execute(OwningActor, EventName);
+			}
+		}
+	};
+
+	/** Using a fake multicast delegate so order can be kept consistent */
+	using FExtensionHandlerEvent = TMap<FDelegateHandle, TSharedRef<FExtensionHandlerRegisteredDelegate>>;
+
 	/** All active component requests. Used to avoid adding the same component twice if requested from multiple sources */
 	TMap<FComponentRequest, int32> RequestTrackingMap;
 
@@ -459,8 +485,14 @@ private:
 
 		/** If this is not null, will only activate for states >= to this */
 		FGameplayTag RequiredInitState;
+
+		/** Indicates this delegate has been removed and will soon be destroyed, do not execute */
+		bool bRemoved;
 	};
 
+	/** Always stored as shared refs to avoid being destroyed during execution */
+	using FActorFeatureDelegateList = TArray<TSharedRef<FActorFeatureRegisteredDelegate>>;
+
 	/** Information for each registered actor */
 	struct FActorFeatureData
 	{
@@ -471,14 +503,17 @@ private:
 		TArray<FActorFeatureState> RegisteredStates;
 
 		/** All delegates bound to this actor */
-		TArray<FActorFeatureRegisteredDelegate> RegisteredDelegates;
+		FActorFeatureDelegateList RegisteredDelegates;
+
+		/** Properly clear delegate list */
+		~FActorFeatureData();
 	};
 
 	/** Actors that were registered as tracking feature state */
 	TMap<FObjectKey, FActorFeatureData> ActorFeatureMap;
 
 	/** Global delegates for any class name */
-	TMap<FComponentRequestReceiverClassPath, TArray<FActorFeatureRegisteredDelegate>> ClassFeatureChangeDelegates;
+	TMap<FComponentRequestReceiverClassPath, FActorFeatureDelegateList> ClassFeatureChangeDelegates;
 
 	/** A queue of state changes to call delegates for, we don't want recursive callbacks */
 	TArray<TPair<AActor*, FActorFeatureState> > StateChangeQueue;
@@ -504,9 +539,9 @@ private:
 	/** Gets or creates the actor struct */
 	FActorFeatureData& FindOrAddActorData(AActor* Actor);
 
-	/** Searches an array of delegates for the one with the specified handle, returns -1 if not found */
-	int32 GetIndexForRegisteredDelegate(TArray<FActorFeatureRegisteredDelegate>& DelegatesToSearch, FDelegateHandle SearchHandle) const;
-	int32 GetIndexForRegisteredDelegate(TArray<FActorFeatureRegisteredDelegate>& DelegatesToSearch, FActorInitStateChangedBPDelegate SearchDelegate) const;
+	/** Try to removes delegate from list, returns true if it removed anything and clear handle if required */
+	bool RemoveActorFeatureDelegateFromList(FActorFeatureDelegateList& DelegateList, FDelegateHandle& SearchHandle) const;
+	bool RemoveActorFeatureDelegateFromList(FActorFeatureDelegateList& DelegateList, FActorInitStateChangedBPDelegate SearchDelegate) const;
 };
 
 #if UE_ENABLE_INCLUDE_ORDER_DEPRECATED_IN_5_2
