Devices - Part 1 of 2 - 30.10 - Adding flags enabling override of provider components added to pawn via GFA component

diff --git a/Engine/Plugins/Runtime/ModularGameplay/Source/ModularGameplay/Private/Components/GameFrameworkComponentManager.cpp b/Engine/Plugins/Runtime/ModularGameplay/Source/ModularGameplay/Private/Components/GameFrameworkComponentManager.cpp
--- a/Engine/Plugins/Runtime/ModularGameplay/Source/ModularGameplay/Private/Components/GameFrameworkComponentManager.cpp
+++ b/Engine/Plugins/Runtime/ModularGameplay/Source/ModularGameplay/Private/Components/GameFrameworkComponentManager.cpp
@@ -481,28 +481,73 @@ EGameFrameworkAddComponentResult UGameFrameworkComponentManager::CreateComponent
 
 	if (!ComponentClass->GetDefaultObject<UActorComponent>()->GetIsReplicated() || ActorInstance->GetLocalRole() == ROLE_Authority)
 	{
-		// If AddUnique is set, it will be added only if no component on ActorInstance is child (or same class) of ComponentClass
-		const bool bAddUnique = EnumHasAnyFlags(AdditionFlags, EGameFrameworkAddComponentFlags::AddUnique);
-		if (bAddUnique)
-		{			
-			if (ActorInstance->GetComponentByClass(ComponentClass))
+		// If ReplaceParents is set, it will remove all parent class components of ComponentClass when it is added. Note: ReplaceParents will supersede other 'AddUnique...' flags.
+		const bool bReplaceParents = EnumHasAnyFlags(AdditionFlags, EGameFrameworkAddComponentFlags::ReplaceParents);
+		if (bReplaceParents)
+		{
+			// Note: Do not add "unique instance" checks here to avoid complicated logic, and because "Replace" behavior and "Add If Unique" behavior are not fully compatible paradigms.
+			// Since this Replace behavior destroys existing parents, we know we will be unique with respect to parent instances after this logic runs.
+
+			TArray<UActorComponent*> FoundParents;
+
+			// Find all components on the ActorInstance whose class is an ancestor of our new component's class
+			const TSet<UActorComponent*>& Components = ActorInstance->GetComponents();
+			for (UActorComponent* ActorComp : Components)
 			{
-				return EGameFrameworkAddComponentResult::Failed;
+				if (ComponentClass->IsChildOf(ActorComp->GetClass()))
+				{
+					FoundParents.Add(ActorComp);
+				}
 			}
+
+			// Remove all ancestor components found
+			for (UActorComponent* ParentComp : FoundParents)
+			{
+				DestroyInstancedComponent(ParentComp);
+			}
+			FoundParents.Reset();
 		}
 
-		// If AddIfNotChild is set,it will be added only if ComponentClass is not a child (or same class) of an existing component on ActorInstance
-		const bool bAddIfNotChild = EnumHasAnyFlags(AdditionFlags, EGameFrameworkAddComponentFlags::AddIfNotChild);
-		if (bAddIfNotChild)
+		// If ReplaceCopiesAndChildren is set, it will remove all components of ComponentClass when it is added. Note: ReplaceCopiesAndChildren will supersede other 'AddUnique...' flags.
+		const bool bReplaceCopiesAndChildren = EnumHasAnyFlags(AdditionFlags, EGameFrameworkAddComponentFlags::ReplaceCopiesAndChildren);
+		if (bReplaceCopiesAndChildren)
 		{
-			const TSet<UActorComponent*>& Components = ActorInstance->GetComponents();
-			for (const UActorComponent* ActorComp : Components)
+			// Note: Do not add "unique instance" checks here to avoid complicated logic, and because "Replace" behavior and "Add If Unique" behavior are not fully compatible paradigms.
+			// Since this Replace behavior destroys existing copies and children, we know we will be unique with respect to instances of our component class after this logic runs.
+
+			// Find and remove all components from ActorInstnace whose class is or descends from ComponentClass
+			while (UActorComponent* FoundRelative = ActorInstance->GetComponentByClass(ComponentClass))
 			{
-				if (ComponentClass->IsChildOf(ActorComp->GetClass()))
+				DestroyInstancedComponent(FoundRelative);
+			}
+		}
+
+		// Only evaluate "AddUnique..." behaviors if we are not using "ReplaceAllParents" or "ReplaceCopiesAndChildren" behaviors
+		if(!bReplaceParents && !bReplaceCopiesAndChildren)
+		{
+			// If AddUnique is set, it will be added only if no component on ActorInstance is child (or same class) of ComponentClass
+			const bool bAddUnique = EnumHasAnyFlags(AdditionFlags, EGameFrameworkAddComponentFlags::AddUnique);
+			if (bAddUnique)
+			{
+				if (ActorInstance->GetComponentByClass(ComponentClass))
 				{
 					return EGameFrameworkAddComponentResult::Failed;
 				}
 			}
+
+			// If AddIfNotChild is set,it will be added only if ComponentClass is not a child of an existing component on ActorInstance
+			const bool bAddIfNotChild = EnumHasAnyFlags(AdditionFlags, EGameFrameworkAddComponentFlags::AddIfNotChild);
+			if (bAddIfNotChild)
+			{
+				const TSet<UActorComponent*>& Components = ActorInstance->GetComponents();
+				for (const UActorComponent* ActorComp : Components)
+				{
+					if (ComponentClass->IsChildOf(ActorComp->GetClass()))
+					{
+						return EGameFrameworkAddComponentResult::Failed;
+					}
+				}
+			}
 		}
 
 		// If UseAutoGeneratedName is set, it will generate a new name and not re-use the class name directly (which can lead to component recycling)
diff --git a/Engine/Plugins/Runtime/ModularGameplay/Source/ModularGameplay/Public/Components/GameFrameworkComponentManager.h b/Engine/Plugins/Runtime/ModularGameplay/Source/ModularGameplay/Public/Components/GameFrameworkComponentManager.h
--- a/Engine/Plugins/Runtime/ModularGameplay/Source/ModularGameplay/Public/Components/GameFrameworkComponentManager.h
+++ b/Engine/Plugins/Runtime/ModularGameplay/Source/ModularGameplay/Public/Components/GameFrameworkComponentManager.h
@@ -21,10 +21,12 @@ UENUM(BlueprintType, meta = (Bitflags, UseEnumValuesAsMaskValuesInEditor = "true
 enum class EGameFrameworkAddComponentFlags : uint8
 {
 	// If 0, no specific rule is applied and the component will be always created on the receiver"),
-	None = 0 UMETA(Hidden),
-	AddUnique = 0b00000001				UMETA(DisplayName = "Add Unique", ToolTip="If AddUnique is set, it will be added only if there is no existing component on ActorInstance of ComponentClass"),
-	AddIfNotChild = 0b00000010			UMETA(DisplayName = "Add Unique If Not Child", ToolTip = "If Add Unique If Not Child is set, it will be added only if ComponentClass is not a child (or same class) of an existing component on ActorInstance"),
-	UseAutoGeneratedName = 0b00000100	UMETA(DisplayName = "Use Auto Generated Name", ToolTip = "If UseAutoGeneratedName is set, it will generate a new name and not re-use the class name directly (which can lead to component recycling)."),
+	None =						0			UMETA(Hidden),
+	AddUnique =					0b00000001	UMETA(DisplayName = "Add Unique", ToolTip="If AddUnique is set, it will be added only if there is no existing component on ActorInstance of ComponentClass."),
+	AddIfNotChild =				0b00000010	UMETA(DisplayName = "Add Unique If Not Child", ToolTip = "If AddIfNotChild is set, it will be added only if ComponentClass is not a child of an existing component on ActorInstance."),
+	UseAutoGeneratedName =		0b00000100	UMETA(DisplayName = "Use Auto Generated Name", ToolTip = "If UseAutoGeneratedName is set, it will generate a new name and not re-use the class name directly (which can lead to component recycling)."),
+	ReplaceParents =			0b00001000	UMETA(DisplayName = "Replace Parents", ToolTip = "If ReplaceParents is set, it will remove all parent class components of ComponentClass when it is added. Note: ReplaceParents will supersede other 'AddUnique...' flags."),
+	ReplaceCopiesAndChildren =	0b00010000	UMETA(DisplayName = "Replace Copies And Children", ToolTip = "If ReplaceCopiesAndChildren is set, it will remove all components of ComponentClass when it is added. Note: ReplaceCopiesAndChildren will supersede other 'AddUnique...' flags."),
 };
 ENUM_CLASS_FLAGS(EGameFrameworkAddComponentFlags);
 
