Stop the game framework component manager from trying to reallocate a component with an incompatible class, as that will crash in AllocateObject.
Add a warning when reallocating components with a compatible class, as this will work but is usually an incorrect setup. Addition flags like AddUnique can be used to stop it from reallocating.

diff --git a/Engine/Plugins/Runtime/ModularGameplay/Source/ModularGameplay/Private/Components/GameFrameworkComponentManager.cpp b/Engine/Plugins/Runtime/ModularGameplay/Source/ModularGameplay/Private/Components/GameFrameworkComponentManager.cpp
--- a/Engine/Plugins/Runtime/ModularGameplay/Source/ModularGameplay/Private/Components/GameFrameworkComponentManager.cpp
+++ b/Engine/Plugins/Runtime/ModularGameplay/Source/ModularGameplay/Private/Components/GameFrameworkComponentManager.cpp
@@ -554,6 +554,33 @@ EGameFrameworkAddComponentResult UGameFrameworkComponentManager::CreateComponent
 		const bool bUseAutoGeneratedName = EnumHasAnyFlags(AdditionFlags, EGameFrameworkAddComponentFlags::UseAutoGeneratedName);
 		const FName NewComponentName = bUseAutoGeneratedName ? NAME_None : ComponentClass->GetFName();
 
+		if (!bUseAutoGeneratedName)
+		{
+			// If we specify the name, this could replace an already existing object. 
+			if (UObject* ExistingObject = StaticFindObjectFast(nullptr, ActorInstance, NewComponentName))
+			{
+				if (ComponentClass->IsChildOf(ExistingObject->GetClass()))
+				{
+					// If they are related classes this will work but probably indicates an improper game feature action setup
+					// If it is intended to add a component from two different places, the creation flags like AddUnique can be used to only add if it does not yet exist
+					UE_LOG(LogModularGameplay, Warning, TEXT("CreateComponentOnInstance is replacing existing object %s with class %s using flags 0x%08x, check for duplicate game feature actions."),
+						*ExistingObject->GetFullName(nullptr, EObjectFullNameFlags::IncludeClassPackage),
+						*ComponentClass->GetPathName(), (int32)AdditionFlags);
+				}
+				else
+				{
+					// This will fatally crash inside AllocateObject, so error and cancel instead. Log and ensure so each case will write to the log
+					UE_LOG(LogModularGameplay, Error, TEXT("CreateComponentOnInstance cannot replace object %s with incompatible class %s using flags 0x%08x!"),
+						*ExistingObject->GetFullName(nullptr, EObjectFullNameFlags::IncludeClassPackage),
+						*ComponentClass->GetPathName(), (int32)AdditionFlags);
+
+					ensureMsgf(false, TEXT("CreateComponentOnInstance cannot replace object with incompatible class, check log for details"));
+
+					return EGameFrameworkAddComponentResult::Failed;
+				}
+			}
+		}
+
 		UActorComponent* NewComp = NewObject<UActorComponent>(ActorInstance, ComponentClass, NewComponentName);
 		TSet<FObjectKey>& ComponentInstances = ComponentClassToComponentInstanceMap.FindOrAdd(*ComponentClass);
 		ComponentInstances.Add(NewComp);
