Clarify usage of bOnlyInGameWorlds param in SendGameFrameworkComponentExtensionEvent, to address concern from PR #11685 and avoid some redundant work

diff --git a/Engine/Plugins/Runtime/ModularGameplay/Source/ModularGameplay/Private/Components/GameFrameworkComponentManager.cpp b/Engine/Plugins/Runtime/ModularGameplay/Source/ModularGameplay/Private/Components/GameFrameworkComponentManager.cpp
--- a/Engine/Plugins/Runtime/ModularGameplay/Source/ModularGameplay/Private/Components/GameFrameworkComponentManager.cpp
+++ b/Engine/Plugins/Runtime/ModularGameplay/Source/ModularGameplay/Private/Components/GameFrameworkComponentManager.cpp
@@ -590,7 +590,8 @@ void UGameFrameworkComponentManager::SendGameFrameworkComponentExtensionEvent(AA
 {
 	if (UGameFrameworkComponentManager* GFCM = GetForActor(Receiver, bOnlyInGameWorlds))
 	{
-		GFCM->SendExtensionEvent(Receiver, EventName);
+		// bOnlyInGameWorlds and the null check is handled by GetForActor, so skip the check in SendExtensionEvent
+		GFCM->SendExtensionEventInternal(Receiver, EventName);
 	}
 }
 
