Fix issue with GameFrameworkComponentManager where if there are multiple requests registered for the same actor/component pair the requests will not properly unregister on destruction.

diff --git a/Engine/Plugins/Runtime/ModularGameplay/Source/ModularGameplay/Private/Components/GameFrameworkComponentManager.cpp b/Engine/Plugins/Runtime/ModularGameplay/Source/ModularGameplay/Private/Components/GameFrameworkComponentManager.cpp
--- a/Engine/Plugins/Runtime/ModularGameplay/Source/ModularGameplay/Private/Components/GameFrameworkComponentManager.cpp
+++ b/Engine/Plugins/Runtime/ModularGameplay/Source/ModularGameplay/Private/Components/GameFrameworkComponentManager.cpp
@@ -295,11 +295,9 @@ TSharedPtr<FComponentRequestHandle> UGameFrameworkComponentManager::AddComponent
 		{
 			// Actor class is not in memory, there will be no actor instances
 		}
-
-		return MakeShared<FComponentRequestHandle>(this, ReceiverClass, ComponentClass);
 	}
 
-	return nullptr;
+	return MakeShared<FComponentRequestHandle>(this, ReceiverClass, ComponentClass);
 }
 
 void UGameFrameworkComponentManager::RemoveComponentRequest(const TSoftClassPtr<AActor>& ReceiverClass, TSubclassOf<UActorComponent> ComponentClass)
diff --git a/Engine/Plugins/Runtime/ModularGameplay/Source/ModularGameplay/Public/Components/GameFrameworkComponentManager.h b/Engine/Plugins/Runtime/ModularGameplay/Source/ModularGameplay/Public/Components/GameFrameworkComponentManager.h
--- a/Engine/Plugins/Runtime/ModularGameplay/Source/ModularGameplay/Public/Components/GameFrameworkComponentManager.h
+++ b/Engine/Plugins/Runtime/ModularGameplay/Source/ModularGameplay/Public/Components/GameFrameworkComponentManager.h
@@ -122,7 +122,7 @@ public:
 	* @param ComponentClass		The component class to add on the new instance
 	* @param AdditionFlags (optional)	Set of flags to considered at component creation
 	* 
-	* @return Returns a handle that will keep the request "alive" until it is destructed, at which point the request is removed. Will return null if there is already an active request with same Receiver & Component classes.
+	* @return Returns a handle that will keep the request "alive" until all handles for the request are destructed, at which point the request is removed.
 	*/
 	UE_API TSharedPtr<FComponentRequestHandle> AddComponentRequest(const TSoftClassPtr<AActor>& ReceiverClass, TSubclassOf<UActorComponent> ComponentClass, const EGameFrameworkAddComponentFlags AdditionFlags = EGameFrameworkAddComponentFlags::None);
 
