Skip execution of GameFrameworkComponentManager unregister callbacks during engine shutdown, as the UWorld maybe in an invalid state

diff --git a/Engine/Plugins/Runtime/ModularGameplay/Source/ModularGameplay/Private/Components/GameFrameworkComponentManager.cpp b/Engine/Plugins/Runtime/ModularGameplay/Source/ModularGameplay/Private/Components/GameFrameworkComponentManager.cpp
--- a/Engine/Plugins/Runtime/ModularGameplay/Source/ModularGameplay/Private/Components/GameFrameworkComponentManager.cpp
+++ b/Engine/Plugins/Runtime/ModularGameplay/Source/ModularGameplay/Private/Components/GameFrameworkComponentManager.cpp
@@ -403,11 +403,11 @@ void UGameFrameworkComponentManager::RemoveExtensionHandler(const TSoftClassPtr<
 		if (ensure(HandlerDelegate))
 		{
 			TSharedRef<FExtensionHandlerRegisteredDelegate> HandlerRef = *HandlerDelegate;
-			// Call it once on unregister
+			// Call it once on unregister, but skip during engine shutdown
 			if (UClass* ReceiverClassPtr = ReceiverClass.Get())
 			{
 				UGameInstance* LocalGameInstance = GetGameInstance();
-				if (ensure(LocalGameInstance))
+				if (!IsEngineExitRequested() && ensure(LocalGameInstance))
 				{
 					UWorld* LocalWorld = LocalGameInstance->GetWorld();
 					ensure(GIsEditor || (LocalWorld != nullptr));
